Fix non-unity

diff --git a/Engine/Plugins/Editor/GeometryMode/Source/GeometryMode/Private/EditorGeometry.cpp b/Engine/Plugins/Editor/GeometryMode/Source/GeometryMode/Private/EditorGeometry.cpp
--- a/Engine/Plugins/Editor/GeometryMode/Source/GeometryMode/Private/EditorGeometry.cpp
+++ b/Engine/Plugins/Editor/GeometryMode/Source/GeometryMode/Private/EditorGeometry.cpp
@@ -6,6 +6,7 @@
 #include "EditorModeManager.h"
 #include "EditorModes.h"
 #include "GeometryEdMode.h"
+#include "GeometryModeModule.h"
 
 DEFINE_LOG_CATEGORY_STATIC(LogEditorGeometry, Log, All);
 
diff --git a/Engine/Plugins/Editor/GeometryMode/Source/GeometryMode/Private/GeometryModeModule.cpp b/Engine/Plugins/Editor/GeometryMode/Source/GeometryMode/Private/GeometryModeModule.cpp
--- a/Engine/Plugins/Editor/GeometryMode/Source/GeometryMode/Private/GeometryModeModule.cpp
+++ b/Engine/Plugins/Editor/GeometryMode/Source/GeometryMode/Private/GeometryModeModule.cpp
@@ -4,6 +4,8 @@
 #include "Modules/ModuleManager.h"
 #include "EditorModeManager.h"
 #include "Editor.h"
+#include "EditorModeRegistry.h"
+#include "GeometryEdMode.h"
 
 
 FEditorModeID FGeometryEditingModes::EM_Geometry = FEditorModeID(TEXT("EM_Geometry"));
diff --git a/Engine/Plugins/Editor/GeometryMode/Source/TextureAlignMode/Private/TextureAlignMode.cpp b/Engine/Plugins/Editor/GeometryMode/Source/TextureAlignMode/Private/TextureAlignMode.cpp
--- a/Engine/Plugins/Editor/GeometryMode/Source/TextureAlignMode/Private/TextureAlignMode.cpp
+++ b/Engine/Plugins/Editor/GeometryMode/Source/TextureAlignMode/Private/TextureAlignMode.cpp
@@ -4,6 +4,7 @@
 #include "EditorModeManager.h"
 #include "EditorModes.h"
 #include "TextureAlignEdMode.h"
+#include "GeometryModeModule.h"
 
 #define LOCTEXT_NAMESPACE "TextureAlignMode"
 
