Missed a non-unity fix

diff --git a/Engine/Plugins/Editor/GeometryMode/Source/GeometryMode/Private/BrushEditingSubsystemImpl.cpp b/Engine/Plugins/Editor/GeometryMode/Source/GeometryMode/Private/BrushEditingSubsystemImpl.cpp
--- a/Engine/Plugins/Editor/GeometryMode/Source/GeometryMode/Private/BrushEditingSubsystemImpl.cpp
+++ b/Engine/Plugins/Editor/GeometryMode/Source/GeometryMode/Private/BrushEditingSubsystemImpl.cpp
@@ -11,6 +11,7 @@
 #include "HitProxies.h"
 #include "LevelViewportClickHandlers.h"
 #include "GeometryModeModule.h"
+#include "ScopedTransaction.h"
 
 DEFINE_LOG_CATEGORY(LogBrushEditing);
 
