#UE4 IWYU fix

diff --git a/Engine/Plugins/Editor/GeometryMode/Source/GeometryMode/Private/GeometryModifiers.cpp b/Engine/Plugins/Editor/GeometryMode/Source/GeometryMode/Private/GeometryModifiers.cpp
--- a/Engine/Plugins/Editor/GeometryMode/Source/GeometryMode/Private/GeometryModifiers.cpp
+++ b/Engine/Plugins/Editor/GeometryMode/Source/GeometryMode/Private/GeometryModifiers.cpp
@@ -42,6 +42,7 @@
 #include "GeomModifier_Optimize.h"
 #include "GeomModifier_Turn.h"
 #include "GeomModifier_Weld.h"
+#include "GeometryModeModule.h"
 
 DEFINE_LOG_CATEGORY_STATIC(LogGeomModifier, Log, All);
 
