#UE4 IWYU fix

diff --git a/Engine/Plugins/Editor/GeometryMode/Source/GeometryMode/Private/GeometryMode.cpp b/Engine/Plugins/Editor/GeometryMode/Source/GeometryMode/Private/GeometryMode.cpp
--- a/Engine/Plugins/Editor/GeometryMode/Source/GeometryMode/Private/GeometryMode.cpp
+++ b/Engine/Plugins/Editor/GeometryMode/Source/GeometryMode/Private/GeometryMode.cpp
@@ -18,6 +18,7 @@
 #include "PropertyEditorModule.h"
 #include "IDetailsView.h"
 #include "GeomModifier.h"
+#include "GeometryModeModule.h"
 
 #define LOCTEXT_NAMESPACE "GeometryMode"
 
