[NO NON-PII COMMIT MESSAGE]

diff --git a/Engine/Plugins/Editor/GeometryMode/Source/GeometryMode/Private/GeometryModifiers.cpp b/Engine/Plugins/Editor/GeometryMode/Source/GeometryMode/Private/GeometryModifiers.cpp
--- a/Engine/Plugins/Editor/GeometryMode/Source/GeometryMode/Private/GeometryModifiers.cpp
+++ b/Engine/Plugins/Editor/GeometryMode/Source/GeometryMode/Private/GeometryModifiers.cpp
@@ -1353,8 +1353,12 @@ void UGeomModifier_Pen::Apply()
 		else
 		{
 			ResultingBrush = FBSPOps::csgAddOperation( BuilderBrush, PF_DefaultFlags, Brush_Add );
+			if (ResultingBrush == nullptr)
+			{
+				return;
+			}
 
-			if (ResultingBrush && ResultingBrush->GetBrushBuilder())
+			if (ResultingBrush->GetBrushBuilder())
 			{
 				FActorLabelUtilities::SetActorLabelUnique(ResultingBrush, FText::Format(NSLOCTEXT("BSPBrushOps", "BrushName", "{0} Brush"), FText::FromString(ResultingBrush->GetBrushBuilder()->GetClass()->GetDescription())).ToString());
 			}
