FGeometryModeToolkit no longer tries to access the current geometry tool when not in geometry editing mode.

diff --git a/Engine/Plugins/Editor/GeometryMode/Source/GeometryMode/Private/GeometryEdMode.cpp b/Engine/Plugins/Editor/GeometryMode/Source/GeometryMode/Private/GeometryEdMode.cpp
--- a/Engine/Plugins/Editor/GeometryMode/Source/GeometryMode/Private/GeometryEdMode.cpp
+++ b/Engine/Plugins/Editor/GeometryMode/Source/GeometryMode/Private/GeometryEdMode.cpp
@@ -1073,9 +1073,9 @@ void FModeTool_GeometryModify::SetCurrentModifier( UGeomModifier* InModifier )
 	CurrentModifier->WasActivated();
 }
 
-UGeomModifier* FModeTool_GeometryModify::GetCurrentModifier() { return CurrentModifier; }
+UGeomModifier* FModeTool_GeometryModify::GetCurrentModifier() const { return CurrentModifier; }
 
-int32 FModeTool_GeometryModify::GetNumModifiers() { return Modifiers.Num(); }
+int32 FModeTool_GeometryModify::GetNumModifiers() const { return Modifiers.Num(); }
 
 FModeTool_GeometryModify::FModeTool_GeometryModify()
 {
diff --git a/Engine/Plugins/Editor/GeometryMode/Source/GeometryMode/Private/GeometryMode.cpp b/Engine/Plugins/Editor/GeometryMode/Source/GeometryMode/Private/GeometryMode.cpp
--- a/Engine/Plugins/Editor/GeometryMode/Source/GeometryMode/Private/GeometryMode.cpp
+++ b/Engine/Plugins/Editor/GeometryMode/Source/GeometryMode/Private/GeometryMode.cpp
@@ -308,9 +308,17 @@ void FGeometryModeToolkit::OnModifierToolBarButtonClicked(UGeomModifier* Modifie
 
 ECheckBoxState FGeometryModeToolkit::IsModifierChecked(UGeomModifier* Modifier) const
 {
-	return (GetGeometryModeTool()->GetCurrentModifier() == Modifier)
-		? ECheckBoxState::Checked
-		: ECheckBoxState::Unchecked;
+	if (const FEdModeGeometry* Mode = (const FEdModeGeometry*) GLevelEditorModeTools().GetActiveMode(FGeometryEditingModes::EM_Geometry))
+	{
+		if (const FModeTool_GeometryModify* Tool = (const FModeTool_GeometryModify*) Mode->GetCurrentTool())
+		{
+			return (Tool->GetCurrentModifier() == Modifier)
+				? ECheckBoxState::Checked
+				: ECheckBoxState::Unchecked;
+		}
+	}
+
+	return ECheckBoxState::Unchecked;
 }
 
 FReply FGeometryModeToolkit::OnApplyClicked()
diff --git a/Engine/Plugins/Editor/GeometryMode/Source/GeometryMode/Public/GeometryEdMode.h b/Engine/Plugins/Editor/GeometryMode/Source/GeometryMode/Public/GeometryEdMode.h
--- a/Engine/Plugins/Editor/GeometryMode/Source/GeometryMode/Public/GeometryEdMode.h
+++ b/Engine/Plugins/Editor/GeometryMode/Source/GeometryMode/Public/GeometryEdMode.h
@@ -219,9 +219,9 @@ public:
 	virtual FString GetName() const override { return TEXT("Modifier"); }
 
 	void SetCurrentModifier( UGeomModifier* InModifier );
-	UGeomModifier* GetCurrentModifier();
+	UGeomModifier* GetCurrentModifier() const;
 
-	int32 GetNumModifiers();
+	int32 GetNumModifiers() const;
 
 	/**
 	 * @return		true if the delta was handled by this editor mode tool.
