Added null checks to all UGeomModifier_*::SupportsSelection() calls that rely on a non-null mode to be returned.

diff --git a/Engine/Plugins/Editor/GeometryMode/Source/GeometryMode/Private/GeometryModifiers.cpp b/Engine/Plugins/Editor/GeometryMode/Source/GeometryMode/Private/GeometryModifiers.cpp
--- a/Engine/Plugins/Editor/GeometryMode/Source/GeometryMode/Private/GeometryModifiers.cpp
+++ b/Engine/Plugins/Editor/GeometryMode/Source/GeometryMode/Private/GeometryModifiers.cpp
@@ -706,8 +706,8 @@ bool UGeomModifier_Extrude::InputDelta(FEditorViewportClient* InViewportClient,
 
 bool UGeomModifier_Extrude::SupportsCurrentSelection()
 {
-	FEdModeGeometry* mode = (FEdModeGeometry*)GLevelEditorModeTools().GetActiveMode(FGeometryEditingModes::EM_Geometry);
-	return mode->HavePolygonsSelected();
+	FEdModeGeometry* Mode = (FEdModeGeometry*) GLevelEditorModeTools().GetActiveMode(FGeometryEditingModes::EM_Geometry);
+	return Mode != nullptr && Mode->HavePolygonsSelected();
 }
 
 void UGeomModifier_Extrude::WasActivated()
@@ -2106,8 +2106,8 @@ bool UGeomModifier_Clip::SupportsCurrentSelection()
 		return false;
 	}
 
-	FEdModeGeometry* mode = (FEdModeGeometry*)GLevelEditorModeTools().GetActiveMode(FGeometryEditingModes::EM_Geometry);
-	return mode->GetSelectionState() ? false : true;
+	FEdModeGeometry* Mode = (FEdModeGeometry*) GLevelEditorModeTools().GetActiveMode(FGeometryEditingModes::EM_Geometry);
+	return Mode != nullptr && (Mode->GetSelectionState() ? false : true);
 }
 
 bool UGeomModifier_Clip::OnApply()
@@ -2485,8 +2485,8 @@ UGeomModifier_Delete::UGeomModifier_Delete(const FObjectInitializer& ObjectIniti
 
 bool UGeomModifier_Delete::SupportsCurrentSelection()
 {
-	FEdModeGeometry* mode = (FEdModeGeometry*)GLevelEditorModeTools().GetActiveMode(FGeometryEditingModes::EM_Geometry);
-	return (mode->HavePolygonsSelected() || mode->HaveVerticesSelected());
+	FEdModeGeometry* Mode = (FEdModeGeometry*) GLevelEditorModeTools().GetActiveMode(FGeometryEditingModes::EM_Geometry);
+	return Mode != nullptr && (Mode->HavePolygonsSelected() || Mode->HaveVerticesSelected());
 }
 
 
@@ -2571,8 +2571,8 @@ UGeomModifier_Create::UGeomModifier_Create(const FObjectInitializer& ObjectIniti
 
 bool UGeomModifier_Create::SupportsCurrentSelection()
 {
-	FEdModeGeometry* mode = (FEdModeGeometry*)GLevelEditorModeTools().GetActiveMode(FGeometryEditingModes::EM_Geometry);
-	return mode->HaveVerticesSelected();
+	FEdModeGeometry* Mode = (FEdModeGeometry*) GLevelEditorModeTools().GetActiveMode(FGeometryEditingModes::EM_Geometry);
+	return Mode != nullptr && Mode->HaveVerticesSelected();
 }
 
 bool UGeomModifier_Create::OnApply()
@@ -2644,8 +2644,8 @@ bool UGeomModifier_Flip::SupportsCurrentSelection()
 	}
 
 	// Supports polygons selected and objects selected
-	FEdModeGeometry* mode = (FEdModeGeometry*)GLevelEditorModeTools().GetActiveMode(FGeometryEditingModes::EM_Geometry);
-	return (!mode->HaveEdgesSelected() && !mode->HaveVerticesSelected());
+	FEdModeGeometry* Mode = (FEdModeGeometry*) GLevelEditorModeTools().GetActiveMode(FGeometryEditingModes::EM_Geometry);
+	return Mode != nullptr && (!Mode->HaveEdgesSelected() && !Mode->HaveVerticesSelected());
 }
 
 bool UGeomModifier_Flip::OnApply()
@@ -2689,24 +2689,23 @@ UGeomModifier_Split::UGeomModifier_Split(const FObjectInitializer& ObjectInitial
 
 bool UGeomModifier_Split::SupportsCurrentSelection()
 {
-	FEdModeGeometry* mode = (FEdModeGeometry*)GLevelEditorModeTools().GetActiveMode(FGeometryEditingModes::EM_Geometry);
+	FEdModeGeometry* Mode = (FEdModeGeometry*)GLevelEditorModeTools().GetActiveMode(FGeometryEditingModes::EM_Geometry);
 
 	// This modifier assumes that a single geometry object is selected
-
-	if( mode->CountObjectsSelected() != 1 )
+	if (Mode == nullptr || Mode->CountObjectsSelected() != 1 )
 	{
 		return false;
 	}
 
-	int32 NumPolygonsSelected = mode->CountSelectedPolygons();
-	int32 NumEdgesSelected = mode->CountSelectedEdges();
-	int32 NumVerticesSelected = mode->CountSelectedVertices();
+	int32 NumPolygonsSelected = Mode->CountSelectedPolygons();
+	int32 NumEdgesSelected = Mode->CountSelectedEdges();
+	int32 NumVerticesSelected = Mode->CountSelectedVertices();
 
-	if( (NumPolygonsSelected == 1 && NumEdgesSelected == 1 && NumVerticesSelected == 0)			// Splitting a face at an edge mid point (scalpel)
-			|| (NumPolygonsSelected == 0 && NumEdgesSelected > 0 && NumVerticesSelected == 0)	// Splitting a brush at an edge mid point (ring cut)
-			|| (NumPolygonsSelected == 1 && NumEdgesSelected == 0 && NumVerticesSelected == 2)	// Splitting a polygon across 2 vertices
-			|| (NumPolygonsSelected == 0 && NumEdgesSelected == 0 && NumVerticesSelected == 2)	// Splitting a brush across 2 vertices
-			)
+	if ((NumPolygonsSelected == 1 && NumEdgesSelected == 1 && NumVerticesSelected == 0)		// Splitting a face at an edge mid point (scalpel)
+		|| (NumPolygonsSelected == 0 && NumEdgesSelected > 0 && NumVerticesSelected == 0)	// Splitting a brush at an edge mid point (ring cut)
+		|| (NumPolygonsSelected == 1 && NumEdgesSelected == 0 && NumVerticesSelected == 2)	// Splitting a polygon across 2 vertices
+		|| (NumPolygonsSelected == 0 && NumEdgesSelected == 0 && NumVerticesSelected == 2)	// Splitting a brush across 2 vertices
+		)
 	{
 		return true;
 	}
@@ -3068,8 +3067,8 @@ bool UGeomModifier_Triangulate::SupportsCurrentSelection()
 		return false;
 	}
 
-	FEdModeGeometry* mode = (FEdModeGeometry*)GLevelEditorModeTools().GetActiveMode(FGeometryEditingModes::EM_Geometry);
-	return (!mode->HaveEdgesSelected() && !mode->HaveVerticesSelected());
+	FEdModeGeometry* Mode = (FEdModeGeometry*) GLevelEditorModeTools().GetActiveMode(FGeometryEditingModes::EM_Geometry);
+	return Mode != nullptr && (!Mode->HaveEdgesSelected() && !Mode->HaveVerticesSelected());
 }
 
 bool UGeomModifier_Triangulate::OnApply()
@@ -3396,8 +3395,8 @@ UGeomModifier_Weld::UGeomModifier_Weld(const FObjectInitializer& ObjectInitializ
 
 bool UGeomModifier_Weld::SupportsCurrentSelection()
 {
-	FEdModeGeometry* mode = (FEdModeGeometry*)GLevelEditorModeTools().GetActiveMode(FGeometryEditingModes::EM_Geometry);
-	return (mode->HaveVerticesSelected() && !mode->HaveEdgesSelected() && !mode->HavePolygonsSelected());
+	FEdModeGeometry* Mode = (FEdModeGeometry*)GLevelEditorModeTools().GetActiveMode(FGeometryEditingModes::EM_Geometry);
+	return Mode != nullptr && (Mode->HaveVerticesSelected() && !Mode->HaveEdgesSelected() && !Mode->HavePolygonsSelected());
 }
 
 bool UGeomModifier_Weld::OnApply()
