GeometryMode: Fix crash in UGeomModifier_Turn::SupportsCurrentSelection()

diff --git a/Engine/Plugins/Editor/GeometryMode/Source/GeometryMode/Private/GeometryModifiers.cpp b/Engine/Plugins/Editor/GeometryMode/Source/GeometryMode/Private/GeometryModifiers.cpp
--- a/Engine/Plugins/Editor/GeometryMode/Source/GeometryMode/Private/GeometryModifiers.cpp
+++ b/Engine/Plugins/Editor/GeometryMode/Source/GeometryMode/Private/GeometryModifiers.cpp
@@ -3226,8 +3226,8 @@ UGeomModifier_Turn::UGeomModifier_Turn(const FObjectInitializer& ObjectInitializ
 
 bool UGeomModifier_Turn::SupportsCurrentSelection()
 {
-	FEdModeGeometry* mode = (FEdModeGeometry*)GLevelEditorModeTools().GetActiveMode(FGeometryEditingModes::EM_Geometry);
-	return mode->HaveEdgesSelected();
+	FEdModeGeometry* Mode = (FEdModeGeometry*)GLevelEditorModeTools().GetActiveMode(FGeometryEditingModes::EM_Geometry);
+	return Mode != nullptr && Mode->HaveEdgesSelected();
 }
 
 bool UGeomModifier_Turn::OnApply()
