GeometryMode: Fix crash when switching modes

diff --git a/Engine/Plugins/Editor/GeometryMode/Source/GeometryMode/Private/GeometryMode.cpp b/Engine/Plugins/Editor/GeometryMode/Source/GeometryMode/Private/GeometryMode.cpp
--- a/Engine/Plugins/Editor/GeometryMode/Source/GeometryMode/Private/GeometryMode.cpp
+++ b/Engine/Plugins/Editor/GeometryMode/Source/GeometryMode/Private/GeometryMode.cpp
@@ -95,11 +95,15 @@ private:
 
 void SGeometryModeControls::SelectionChanged()
 {
-	FModeTool_GeometryModify* ModeTool = GeometryModeWeakPtr.Pin()->GetGeometryModeTool();
-
-	if (!ModeTool->GetCurrentModifier()->SupportsCurrentSelection() && ModeTool->GetNumModifiers() > 0)
+	if (const TSharedPtr<FGeometryModeToolkit> GeometryModeSharedPtr = GeometryModeWeakPtr.Pin())
 	{
-		ModeTool->SetCurrentModifier(ModeTool->GetModifier(0));
+		if (FModeTool_GeometryModify* ModeTool = GeometryModeSharedPtr->GetGeometryModeTool())
+		{
+			if (!ModeTool->GetCurrentModifier()->SupportsCurrentSelection() && ModeTool->GetNumModifiers() > 0)
+			{
+				ModeTool->SetCurrentModifier(ModeTool->GetModifier(0));
+			}
+		}
 	}
 }
 
@@ -107,15 +111,15 @@ void SGeometryModeControls::Construct(const FArguments& InArgs, TSharedRef<FGeom
 {
 	GeometryModeWeakPtr = GeometryMode;
 
-	FModeTool_GeometryModify* ModeTool = GeometryMode->GetGeometryModeTool();
-
-	
-	if (ModeTool->GetNumModifiers() > 0)
+	if (FModeTool_GeometryModify* ModeTool = GeometryMode->GetGeometryModeTool())
 	{
-		ModeTool->SetCurrentModifier(ModeTool->GetModifier(0));
-	}
+		if (ModeTool->GetNumModifiers() > 0)
+		{
+			ModeTool->SetCurrentModifier(ModeTool->GetModifier(0));
+		}
 
-	CreateLayout(GeometryMode);
+		CreateLayout(GeometryMode);
+	}
 }
 
 void SGeometryModeControls::SetPropertyObjects(const TArray<UObject*>& PropertyObjects)
@@ -125,9 +129,14 @@ void SGeometryModeControls::SetPropertyObjects(const TArray<UObject*>& PropertyO
 
 bool SGeometryModeControls::ArePropertiesVisible() const
 {
-	FModeTool_GeometryModify* ModeTool = GeometryModeWeakPtr.Pin()->GetGeometryModeTool();
-
-	return (ModeTool->GetNumModifiers() > 0) && (ModeTool->GetCurrentModifier() != ModeTool->GetModifier(0));
+	if (const TSharedPtr<FGeometryModeToolkit> GeometryModeSharedPtr = GeometryModeWeakPtr.Pin())
+	{
+		if (FModeTool_GeometryModify* ModeTool = GeometryModeSharedPtr->GetGeometryModeTool())
+		{
+			return (ModeTool->GetNumModifiers() > 0) && (ModeTool->GetCurrentModifier() != ModeTool->GetModifier(0));
+		}
+	}
+	return false;
 }
 
 EVisibility SGeometryModeControls::GetPropertyVisibility() const
@@ -277,33 +286,38 @@ void FGeometryModeToolkit::OnGeometrySelectionChanged()
 
 class FModeTool_GeometryModify* FGeometryModeToolkit::GetGeometryModeTool() const
 {
-	FEdModeGeometry* Mode = (FEdModeGeometry*)GLevelEditorModeTools().GetActiveMode(FGeometryEditingModes::EM_Geometry);
-	FModeTool* Tool = Mode ? Mode->GetCurrentTool() : NULL;
-
-	check(Tool);
+	if (FEdModeGeometry* Mode = static_cast<FEdModeGeometry*>(GLevelEditorModeTools().GetActiveMode(FGeometryEditingModes::EM_Geometry)))
+	{
+		return static_cast<FModeTool_GeometryModify*>(Mode->GetCurrentTool());
+	}
 
-	return (FModeTool_GeometryModify*)Tool;
+	return nullptr;
 }
 
 void FGeometryModeToolkit::OnModifierStateChanged(ECheckBoxState NewCheckedState, UGeomModifier* Modifier)
 {
 	if (NewCheckedState == ECheckBoxState::Checked)
 	{
-		GetGeometryModeTool()->SetCurrentModifier(Modifier);
+		if (FModeTool_GeometryModify* ModeTool = GetGeometryModeTool())
+		{
+			ModeTool->SetCurrentModifier(Modifier);
 
-		TArray<UObject*> PropertyObjects;
-		PropertyObjects.Add(GetGeometryModeTool()->GetCurrentModifier());
+			TArray<UObject*> PropertyObjects;
+			PropertyObjects.Add(ModeTool->GetCurrentModifier());
 
-		GeomWidget->SetPropertyObjects(PropertyObjects);
+			GeomWidget->SetPropertyObjects(PropertyObjects);
+		}
 	}
 }
 
 void FGeometryModeToolkit::OnModifierToolBarButtonClicked(UGeomModifier* Modifier)
 {
-	const ECheckBoxState NewState = GetGeometryModeTool()->GetCurrentModifier() == Modifier ? ECheckBoxState::Unchecked : ECheckBoxState::Checked;
-
-	OnModifierStateChanged(NewState, Modifier);
+	if (const FModeTool_GeometryModify* ModeTool = GetGeometryModeTool())
+	{
+		const ECheckBoxState NewState = ModeTool->GetCurrentModifier() == Modifier ? ECheckBoxState::Unchecked : ECheckBoxState::Checked;
 
+		OnModifierStateChanged(NewState, Modifier);
+	}
 }
 
 ECheckBoxState FGeometryModeToolkit::IsModifierChecked(UGeomModifier* Modifier) const
@@ -325,9 +339,13 @@ FReply FGeometryModeToolkit::OnApplyClicked()
 {
 	check(GLevelEditorModeTools().IsModeActive(FGeometryEditingModes::EM_Geometry));
 
-	GetGeometryModeTool()->GetCurrentModifier()->Apply();
+	if (const FModeTool_GeometryModify* ModeTool = GetGeometryModeTool())
+	{
+		ModeTool->GetCurrentModifier()->Apply();
 
-	return FReply::Handled();
+		return FReply::Handled();
+	}
+	return FReply::Unhandled();
 }
 
 void FGeometryModeToolkit::OnModifierClicked(UGeomModifier* Modifier)
@@ -368,6 +386,10 @@ void FGeometryModeToolkit::BuildToolPalette(FName Palette, class FToolBarBuilder
 	if (Palette == GeometryModePaletteNames::ToolsPalette)
 	{
 		FModeTool_GeometryModify* GeometryModeTool = GetGeometryModeTool();
+		if (!GeometryModeTool)
+		{
+			return;
+		}
 
 		// Loop through all geometry modifiers and create radio buttons for ones with the bPushButton set to false
 		//int32 CurrentModifierButtonCount = 0;
@@ -412,15 +434,25 @@ void FGeometryModeToolkit::BuildToolPalette(FName Palette, class FToolBarBuilder
 
 FText FGeometryModeToolkit::GetActiveToolDisplayName() const
 {
-	return GetGeometryModeTool()->GetCurrentModifier()->GetModifierDescription();//  + LOCTEXT("GeometryMode_ToolsTool", " Tool");
+	if (const FModeTool_GeometryModify* ModeTool = GetGeometryModeTool())
+	{
+		return ModeTool->GetCurrentModifier()->GetModifierDescription();//  + LOCTEXT("GeometryMode_ToolsTool", " Tool");
+	}
+
+	return FText();
 }
 
 FText FGeometryModeToolkit::GetActiveToolMessage() const
 {
-	UGeomModifier* CurrentModifier = GetGeometryModeTool()->GetCurrentModifier();
-	return ( CurrentModifier && CurrentModifier->SupportsCurrentSelection() )
-		? CurrentModifier->GetModifierTooltip()
-		: LOCTEXT("GeometryMode_NoBrushSelectedMessage", "Select a brush actor to begin editing its geometry");
+	if (const FModeTool_GeometryModify* ModeTool = GetGeometryModeTool())
+	{
+		UGeomModifier* CurrentModifier = ModeTool->GetCurrentModifier();
+		return ( CurrentModifier && CurrentModifier->SupportsCurrentSelection() )
+			? CurrentModifier->GetModifierTooltip()
+			: LOCTEXT("GeometryMode_NoBrushSelectedMessage", "Select a brush actor to begin editing its geometry");
+	}
+
+	return FText();
 }
 
 #undef LOCTEXT_NAMESPACE
