Improve support for an extended WORLD_MAX

diff --git a/Engine/Plugins/Animation/IKRig/Source/IKRigEditor/Public/RetargetEditor/SRetargetAnimAssetsWindow.h b/Engine/Plugins/Animation/IKRig/Source/IKRigEditor/Public/RetargetEditor/SRetargetAnimAssetsWindow.h
--- a/Engine/Plugins/Animation/IKRig/Source/IKRigEditor/Public/RetargetEditor/SRetargetAnimAssetsWindow.h
+++ b/Engine/Plugins/Animation/IKRig/Source/IKRigEditor/Public/RetargetEditor/SRetargetAnimAssetsWindow.h
@@ -104,7 +104,7 @@ public:
 		DrawHelper.GridColorAxis = FColor(70, 70, 70);
 		DrawHelper.GridColorMajor = FColor(40, 40, 40);
 		DrawHelper.GridColorMinor =  FColor(20, 20, 20);
-		DrawHelper.PerspectiveGridSize = HALF_WORLD_MAX1;
+		DrawHelper.PerspectiveGridSize = UE_OLD_HALF_WORLD_MAX1;
 
 		bDisableInput = true;
 	}
diff --git a/Engine/Plugins/Editor/GeometryMode/Source/GeometryMode/Private/GeometryModifiers.cpp b/Engine/Plugins/Editor/GeometryMode/Source/GeometryMode/Private/GeometryModifiers.cpp
--- a/Engine/Plugins/Editor/GeometryMode/Source/GeometryMode/Private/GeometryModifiers.cpp
+++ b/Engine/Plugins/Editor/GeometryMode/Source/GeometryMode/Private/GeometryModifiers.cpp
@@ -1826,82 +1826,82 @@ static void BuildGiantAlignedBrush( ABrush& OutGiantBrush, const FPlane& InPlane
 	OutGiantBrush.Brush->Polys->Element.Empty();
 
 	// Create a list of vertices that can be used for the new brush
-	FVector vtxs[8];
+	FVector3f vtxs[8];
 
 	FPlane FlippedPlane = InPlane.Flip();
 	FPoly TempPoly = FPoly::BuildInfiniteFPoly( FlippedPlane );
 	TempPoly.Finalize(&OutGiantBrush,0);
-	vtxs[0] = (FVector)TempPoly.Vertices[0];
-	vtxs[1] = (FVector)TempPoly.Vertices[1];
-	vtxs[2] = (FVector)TempPoly.Vertices[2];
-	vtxs[3] = (FVector)TempPoly.Vertices[3];
+	vtxs[0] = TempPoly.Vertices[0];
+	vtxs[1] = TempPoly.Vertices[1];
+	vtxs[2] = TempPoly.Vertices[2];
+	vtxs[3] = TempPoly.Vertices[3];
 
 	FlippedPlane = FlippedPlane.Flip();
 	FPoly TempPoly2 = FPoly::BuildInfiniteFPoly( FlippedPlane );
-	vtxs[4] = FVector(TempPoly2.Vertices[0] + (TempPoly2.Normal * -(WORLD_MAX)));	vtxs[5] = FVector(TempPoly2.Vertices[1] + (TempPoly2.Normal * -(WORLD_MAX)));
-	vtxs[6] = FVector(TempPoly2.Vertices[2] + (TempPoly2.Normal * -(WORLD_MAX)));	vtxs[7] = FVector(TempPoly2.Vertices[3] + (TempPoly2.Normal * -(WORLD_MAX)));
+	vtxs[4] = FVector3f(TempPoly2.Vertices[0] + (TempPoly2.Normal * -(UE_OLD_WORLD_MAX)));	vtxs[5] = FVector3f(TempPoly2.Vertices[1] + (TempPoly2.Normal * -(UE_OLD_WORLD_MAX)));
+	vtxs[6] = FVector3f(TempPoly2.Vertices[2] + (TempPoly2.Normal * -(UE_OLD_WORLD_MAX)));	vtxs[7] = FVector3f(TempPoly2.Vertices[3] + (TempPoly2.Normal * -(UE_OLD_WORLD_MAX)));
 
 	// Create the polys for the new brush.
 	FPoly newPoly;
 
 	// TOP
 	newPoly.Init();
-	newPoly.Base = (FVector3f)vtxs[0];
-	newPoly.Vertices.Add( (FVector3f)vtxs[0] );
-	newPoly.Vertices.Add( (FVector3f)vtxs[1] );
-	newPoly.Vertices.Add( (FVector3f)vtxs[2] );
-	newPoly.Vertices.Add( (FVector3f)vtxs[3] );
+	newPoly.Base = vtxs[0];
+	newPoly.Vertices.Add( vtxs[0] );
+	newPoly.Vertices.Add( vtxs[1] );
+	newPoly.Vertices.Add( vtxs[2] );
+	newPoly.Vertices.Add( vtxs[3] );
 	newPoly.Finalize(&OutGiantBrush,0);
 	new(OutGiantBrush.Brush->Polys->Element)FPoly(newPoly);
 
 	// BOTTOM
 	newPoly.Init();
-	newPoly.Base = (FVector3f)vtxs[4];
-	newPoly.Vertices.Add( (FVector3f)vtxs[4] );
-	newPoly.Vertices.Add( (FVector3f)vtxs[5] );
-	newPoly.Vertices.Add( (FVector3f)vtxs[6] );
-	newPoly.Vertices.Add( (FVector3f)vtxs[7] );
+	newPoly.Base = vtxs[4];
+	newPoly.Vertices.Add( vtxs[4] );
+	newPoly.Vertices.Add( vtxs[5] );
+	newPoly.Vertices.Add( vtxs[6] );
+	newPoly.Vertices.Add( vtxs[7] );
 	newPoly.Finalize(&OutGiantBrush,0);
 	new(OutGiantBrush.Brush->Polys->Element)FPoly(newPoly);
 
 	// SIDES
 	// 1
 	newPoly.Init();
-	newPoly.Base = (FVector3f)vtxs[1];
-	newPoly.Vertices.Add( (FVector3f)vtxs[1] );
-	newPoly.Vertices.Add( (FVector3f)vtxs[0] );
-	newPoly.Vertices.Add( (FVector3f)vtxs[7] );
-	newPoly.Vertices.Add( (FVector3f)vtxs[6] );
+	newPoly.Base = vtxs[1];
+	newPoly.Vertices.Add( vtxs[1] );
+	newPoly.Vertices.Add( vtxs[0] );
+	newPoly.Vertices.Add( vtxs[7] );
+	newPoly.Vertices.Add( vtxs[6] );
 	newPoly.Finalize(&OutGiantBrush,0);
 	new(OutGiantBrush.Brush->Polys->Element)FPoly(newPoly);
 
 	// 2
 	newPoly.Init();
-	newPoly.Base = (FVector3f)vtxs[2];
-	newPoly.Vertices.Add( (FVector3f)vtxs[2] );
-	newPoly.Vertices.Add( (FVector3f)vtxs[1] );
-	newPoly.Vertices.Add( (FVector3f)vtxs[6] );
-	newPoly.Vertices.Add( (FVector3f)vtxs[5] );
+	newPoly.Base = vtxs[2];
+	newPoly.Vertices.Add( vtxs[2] );
+	newPoly.Vertices.Add( vtxs[1] );
+	newPoly.Vertices.Add( vtxs[6] );
+	newPoly.Vertices.Add( vtxs[5] );
 	newPoly.Finalize(&OutGiantBrush,0);
 	new(OutGiantBrush.Brush->Polys->Element)FPoly(newPoly);
 
 	// 3
 	newPoly.Init();
-	newPoly.Base = (FVector3f)vtxs[3];
-	newPoly.Vertices.Add( (FVector3f)vtxs[3] );
-	newPoly.Vertices.Add( (FVector3f)vtxs[2] );
-	newPoly.Vertices.Add( (FVector3f)vtxs[5] );
-	newPoly.Vertices.Add( (FVector3f)vtxs[4] );
+	newPoly.Base = vtxs[3];
+	newPoly.Vertices.Add( vtxs[3] );
+	newPoly.Vertices.Add( vtxs[2] );
+	newPoly.Vertices.Add( vtxs[5] );
+	newPoly.Vertices.Add( vtxs[4] );
 	newPoly.Finalize(&OutGiantBrush,0);
 	new(OutGiantBrush.Brush->Polys->Element)FPoly(newPoly);
 
 	// 4
 	newPoly.Init();
-	newPoly.Base = (FVector3f)vtxs[0];
-	newPoly.Vertices.Add( (FVector3f)vtxs[0] );
-	newPoly.Vertices.Add( (FVector3f)vtxs[3] );
-	newPoly.Vertices.Add( (FVector3f)vtxs[4] );
-	newPoly.Vertices.Add( (FVector3f)vtxs[7] );
+	newPoly.Base = vtxs[0];
+	newPoly.Vertices.Add( vtxs[0] );
+	newPoly.Vertices.Add( vtxs[3] );
+	newPoly.Vertices.Add( vtxs[4] );
+	newPoly.Vertices.Add( vtxs[7] );
 	newPoly.Finalize(&OutGiantBrush,0);
 	new(OutGiantBrush.Brush->Polys->Element)FPoly(newPoly);
 
diff --git a/Engine/Plugins/Experimental/Interchange/Runtime/Source/Import/Private/Mesh/InterchangeStaticMeshFactory.cpp b/Engine/Plugins/Experimental/Interchange/Runtime/Source/Import/Private/Mesh/InterchangeStaticMeshFactory.cpp
--- a/Engine/Plugins/Experimental/Interchange/Runtime/Source/Import/Private/Mesh/InterchangeStaticMeshFactory.cpp
+++ b/Engine/Plugins/Experimental/Interchange/Runtime/Source/Import/Private/Mesh/InterchangeStaticMeshFactory.cpp
@@ -1247,10 +1247,10 @@ bool UInterchangeStaticMeshFactory::GenerateKDopCollision(const FCreateAssetPara
 
 		FVector3f Base = Planes[PlaneIndex] * Planes[PlaneIndex].W;
 
-		Polygon.Vertices.Emplace(Base + AxisX * HALF_WORLD_MAX + AxisY * HALF_WORLD_MAX);
-		Polygon.Vertices.Emplace(Base + AxisX * HALF_WORLD_MAX - AxisY * HALF_WORLD_MAX);
-		Polygon.Vertices.Emplace(Base - AxisX * HALF_WORLD_MAX - AxisY * HALF_WORLD_MAX);
-		Polygon.Vertices.Emplace(Base - AxisX * HALF_WORLD_MAX + AxisY * HALF_WORLD_MAX);
+		Polygon.Vertices.Emplace(Base + AxisX * UE_OLD_HALF_WORLD_MAX + AxisY * UE_OLD_HALF_WORLD_MAX);
+		Polygon.Vertices.Emplace(Base + AxisX * UE_OLD_HALF_WORLD_MAX - AxisY * UE_OLD_HALF_WORLD_MAX);
+		Polygon.Vertices.Emplace(Base - AxisX * UE_OLD_HALF_WORLD_MAX - AxisY * UE_OLD_HALF_WORLD_MAX);
+		Polygon.Vertices.Emplace(Base - AxisX * UE_OLD_HALF_WORLD_MAX + AxisY * UE_OLD_HALF_WORLD_MAX);
 
 		for (int32 OtherPlaneIndex = 0; OtherPlaneIndex < Planes.Num(); OtherPlaneIndex++)
 		{
diff --git a/Engine/Plugins/Experimental/Water/Source/Runtime/Private/WaterInfoRendering.cpp b/Engine/Plugins/Experimental/Water/Source/Runtime/Private/WaterInfoRendering.cpp
--- a/Engine/Plugins/Experimental/Water/Source/Runtime/Private/WaterInfoRendering.cpp
+++ b/Engine/Plugins/Experimental/Water/Source/Runtime/Private/WaterInfoRendering.cpp
@@ -257,14 +257,14 @@ static FMatrix BuildOrthoMatrix(float InOrthoWidth, float InOrthoHeight)
 {
 	check((int32)ERHIZBuffer::IsInverted);
 
-	const float OrthoWidth = InOrthoWidth / 2.0f;
-	const float OrthoHeight = InOrthoHeight / 2.0f;
+	const FMatrix::FReal OrthoWidth = InOrthoWidth / 2.0f;
+	const FMatrix::FReal OrthoHeight = InOrthoHeight / 2.0f;
 
-	const float NearPlane = 0.f;
-	const float FarPlane = WORLD_MAX / 8.0f;
+	const FMatrix::FReal NearPlane = 0.f;
+	const FMatrix::FReal FarPlane = UE_FLOAT_HUGE_DISTANCE / 4.0f;
 
-	const float ZScale = 1.0f / (FarPlane - NearPlane);
-	const float ZOffset = 0;
+	const FMatrix::FReal ZScale = 1.0f / (FarPlane - NearPlane);
+	const FMatrix::FReal ZOffset = 0;
 
 	return FReversedZOrthoMatrix(
 		OrthoWidth,
diff --git a/Engine/Plugins/FX/Niagara/Source/Niagara/Private/NiagaraBakerSettings.cpp b/Engine/Plugins/FX/Niagara/Source/Niagara/Private/NiagaraBakerSettings.cpp
--- a/Engine/Plugins/FX/Niagara/Source/Niagara/Private/NiagaraBakerSettings.cpp
+++ b/Engine/Plugins/FX/Niagara/Source/Niagara/Private/NiagaraBakerSettings.cpp
@@ -118,7 +118,7 @@ FMatrix UNiagaraBakerSettings::GetProjectionMatrix() const
 	}
 	else
 	{
-		const float ZRange = WORLD_MAX;
+		const FMatrix::FReal ZRange = UE_OLD_WORLD_MAX;
 		return FReversedZOrthoMatrix(CurrentCamera.OrthoWidth / 2.0f, CurrentCamera.OrthoWidth * AspectRatioY / 2.0f, 0.5f / ZRange, ZRange);
 	}
 }
diff --git a/Engine/Plugins/FX/Niagara/Source/Niagara/Private/NiagaraComponent.cpp b/Engine/Plugins/FX/Niagara/Source/Niagara/Private/NiagaraComponent.cpp
--- a/Engine/Plugins/FX/Niagara/Source/Niagara/Private/NiagaraComponent.cpp
+++ b/Engine/Plugins/FX/Niagara/Source/Niagara/Private/NiagaraComponent.cpp
@@ -166,10 +166,10 @@ static FAutoConsoleVariableRef CVarNiagaraCanPreventScalabilityCullingOnPlayerFX
 
 bool FitsIntoFloat(FVector InValue)
 {
-	// copied from PlayerCameraManager, 1048576 is HALF_WORLD_MAX for float precision vectors
-	return (InValue.X < 1048576.0f && InValue.X > -1048576.0f &&
-			InValue.Y < 1048576.0f && InValue.Y > -1048576.0f &&
-			InValue.Z < 1048576.0f && InValue.Z > -1048576.0f);
+	// copied from PlayerCameraManager,use HALF_WORLD_MAX for float precision vectors
+	return (InValue.X < UE_OLD_HALF_WORLD_MAX  && InValue.X > -UE_OLD_HALF_WORLD_MAX &&
+			InValue.Y < UE_OLD_HALF_WORLD_MAX && InValue.Y > -UE_OLD_HALF_WORLD_MAX &&
+			InValue.Z < UE_OLD_HALF_WORLD_MAX && InValue.Z > -UE_OLD_HALF_WORLD_MAX);
 }
 
 FNiagaraSceneProxy::FNiagaraSceneProxy(UNiagaraComponent* InComponent)
diff --git a/Engine/Plugins/FX/Niagara/Source/Niagara/Private/NiagaraRenderer.cpp b/Engine/Plugins/FX/Niagara/Source/Niagara/Private/NiagaraRenderer.cpp
--- a/Engine/Plugins/FX/Niagara/Source/Niagara/Private/NiagaraRenderer.cpp
+++ b/Engine/Plugins/FX/Niagara/Source/Niagara/Private/NiagaraRenderer.cpp
@@ -903,12 +903,12 @@ FVector4f FNiagaraRenderer::CalcMacroUVParameters(const FSceneView& View, FVecto
 		const FVector4 RightPostProjectionPosition = ViewProjMatrix.TransformPosition(MacroUVPosition + MacroUVRadius * ViewMatrix.GetColumn(0));
 		const FVector4 UpPostProjectionPosition = ViewProjMatrix.TransformPosition(MacroUVPosition + MacroUVRadius * ViewMatrix.GetColumn(1));
 
-		const float RightNDCPosX = RightPostProjectionPosition.X / FMath::Max(RightPostProjectionPosition.W, 0.0001f);
-		const float UpNDCPosY = UpPostProjectionPosition.Y / FMath::Max(UpPostProjectionPosition.W, 0.0001f);
-		const float DX = FMath::Min<float>(RightNDCPosX - ObjectNDCPosition.X, WORLD_MAX);
-		const float DY = FMath::Min<float>(UpNDCPosY - ObjectNDCPosition.Y, WORLD_MAX);
+		const FVector4::FReal RightNDCPosX = RightPostProjectionPosition.X / FMath::Max(RightPostProjectionPosition.W, 0.0001f);
+		const FVector4::FReal UpNDCPosY = UpPostProjectionPosition.Y / FMath::Max(UpPostProjectionPosition.W, 0.0001f);
+		const FVector4::FReal DX = FMath::Min(RightNDCPosX - ObjectNDCPosition.X, WORLD_MAX);
+		const FVector4::FReal DY = FMath::Min(UpNDCPosY - ObjectNDCPosition.Y, WORLD_MAX);
 
-		MacroUVParameters.X = float(ObjectNDCPosition.X);
+		MacroUVParameters.X = float(ObjectNDCPosition.X);	// LWC_TODO: Precision loss?
 		MacroUVParameters.Y = float(ObjectNDCPosition.Y);
 		if (DX != 0.0f && DY != 0.0f && !FMath::IsNaN(DX) && FMath::IsFinite(DX) && !FMath::IsNaN(DY) && FMath::IsFinite(DY))
 		{
diff --git a/Engine/Plugins/FX/Niagara/Source/Niagara/Private/NiagaraSystemInstance.cpp b/Engine/Plugins/FX/Niagara/Source/Niagara/Private/NiagaraSystemInstance.cpp
--- a/Engine/Plugins/FX/Niagara/Source/Niagara/Private/NiagaraSystemInstance.cpp
+++ b/Engine/Plugins/FX/Niagara/Source/Niagara/Private/NiagaraSystemInstance.cpp
@@ -1751,10 +1751,10 @@ float FNiagaraSystemInstance::GetLODDistance()
 	if (WorldManager->GetCachedViewInfo().Num() > 0)
 	{
 		// We are being ticked inside the WorldManager and can safely use the list of cached player view locations
-		float LODDistanceSqr = FMath::Square(WORLD_MAX);
+		FVector::FReal LODDistanceSqr = FMath::Square(WORLD_MAX);
 		for (const FNiagaraCachedViewInfo& ViewInfo : WorldManager->GetCachedViewInfo())
 		{
-			const float DistanceToEffectSqr = FVector(ViewInfo.ViewToWorld.GetOrigin() - EffectLocation).SizeSquared();
+			const FVector::FReal DistanceToEffectSqr = FVector(ViewInfo.ViewToWorld.GetOrigin() - EffectLocation).SizeSquared();
 			LODDistanceSqr = FMath::Min(LODDistanceSqr, DistanceToEffectSqr);
 		}
 		LODDistance = FMath::Sqrt(LODDistanceSqr);
@@ -1785,10 +1785,10 @@ float FNiagaraSystemInstance::GetLODDistance()
 
 		if (PlayerViewLocations.Num() > 0)
 		{
-			float LODDistanceSqr = FMath::Square(WORLD_MAX);
+			FVector::FReal LODDistanceSqr = FMath::Square(WORLD_MAX);
 			for (const FVector& ViewLocation : PlayerViewLocations)
 			{
-				const float DistanceToEffectSqr = FVector(ViewLocation - EffectLocation).SizeSquared();
+				const FVector::FReal DistanceToEffectSqr = FVector(ViewLocation - EffectLocation).SizeSquared();
 				LODDistanceSqr = FMath::Min(LODDistanceSqr, DistanceToEffectSqr);
 			}
 			LODDistance = FMath::Sqrt(LODDistanceSqr);
diff --git a/Engine/Plugins/FX/Niagara/Source/NiagaraEditor/Private/Widgets/SNiagaraSystemViewport.cpp b/Engine/Plugins/FX/Niagara/Source/NiagaraEditor/Private/Widgets/SNiagaraSystemViewport.cpp
--- a/Engine/Plugins/FX/Niagara/Source/NiagaraEditor/Private/Widgets/SNiagaraSystemViewport.cpp
+++ b/Engine/Plugins/FX/Niagara/Source/NiagaraEditor/Private/Widgets/SNiagaraSystemViewport.cpp
@@ -85,7 +85,7 @@ FNiagaraSystemViewportClient::FNiagaraSystemViewportClient(FAdvancedPreviewScene
 	DrawHelper.GridColorAxis = FColor(80,80,80);
 	DrawHelper.GridColorMajor = FColor(72,72,72);
 	DrawHelper.GridColorMinor = FColor(64,64,64);
-	DrawHelper.PerspectiveGridSize = HALF_WORLD_MAX1;
+	DrawHelper.PerspectiveGridSize = UE_OLD_HALF_WORLD_MAX1;
 	ShowWidget(false);
 
 	FEditorViewportClient::SetViewMode(VMI_Lit);
@@ -859,7 +859,7 @@ FNiagaraBaselineViewportClient::FNiagaraBaselineViewportClient(FAdvancedPreviewS
 	DrawHelper.GridColorAxis = FColor(80,80,80);
 	DrawHelper.GridColorMajor = FColor(72,72,72);
 	DrawHelper.GridColorMinor = FColor(64,64,64);
-	DrawHelper.PerspectiveGridSize = HALF_WORLD_MAX1;
+	DrawHelper.PerspectiveGridSize = UE_OLD_HALF_WORLD_MAX1;
 	ShowWidget(false);
 
 	SetViewMode(VMI_Lit);
diff --git a/Engine/Plugins/Runtime/HairStrands/Source/HairStrandsEditor/Private/Widgets/SGroomEditorViewport.cpp b/Engine/Plugins/Runtime/HairStrands/Source/HairStrandsEditor/Private/Widgets/SGroomEditorViewport.cpp
--- a/Engine/Plugins/Runtime/HairStrands/Source/HairStrandsEditor/Private/Widgets/SGroomEditorViewport.cpp
+++ b/Engine/Plugins/Runtime/HairStrands/Source/HairStrandsEditor/Private/Widgets/SGroomEditorViewport.cpp
@@ -61,7 +61,7 @@ FGroomEditorViewportClient::FGroomEditorViewportClient(FAdvancedPreviewScene& In
 	DrawHelper.GridColorAxis = FColor(80,80,80);
 	DrawHelper.GridColorMajor = FColor(72,72,72);
 	DrawHelper.GridColorMinor = FColor(64,64,64);
-	DrawHelper.PerspectiveGridSize = HALF_WORLD_MAX1;
+	DrawHelper.PerspectiveGridSize = UE_OLD_HALF_WORLD_MAX1;
 	ShowWidget(false);
 
 	SetViewMode(VMI_Lit);
diff --git a/Engine/Plugins/Runtime/MixedRealityCaptureFramework/Source/MixedRealityCaptureFramework/Private/MixedRealityCaptureComponent.cpp b/Engine/Plugins/Runtime/MixedRealityCaptureFramework/Source/MixedRealityCaptureFramework/Private/MixedRealityCaptureComponent.cpp
--- a/Engine/Plugins/Runtime/MixedRealityCaptureFramework/Source/MixedRealityCaptureFramework/Private/MixedRealityCaptureComponent.cpp
+++ b/Engine/Plugins/Runtime/MixedRealityCaptureFramework/Source/MixedRealityCaptureFramework/Private/MixedRealityCaptureComponent.cpp
@@ -565,7 +565,7 @@ bool UMixedRealityCaptureComponent::GetEditorPreviewInfo(float /*DeltaTime*/, FM
 
 	// see BuildProjectionMatrix() in SceneCaptureRendering.cpp
 	ViewOut.OrthoNearClipPlane = 0.0f;
-	ViewOut.OrthoFarClipPlane  = WORLD_MAX / 8.0f;;
+	ViewOut.OrthoFarClipPlane  = UE_FLOAT_HUGE_DISTANCE / 4.0f;
 
 	ViewOut.PostProcessBlendWeight = PostProcessBlendWeight;
 	if (PostProcessBlendWeight > 0.0f)
diff --git a/Engine/Plugins/Runtime/ReplicationGraph/Source/Private/ReplicationGraph.cpp b/Engine/Plugins/Runtime/ReplicationGraph/Source/Private/ReplicationGraph.cpp
--- a/Engine/Plugins/Runtime/ReplicationGraph/Source/Private/ReplicationGraph.cpp
+++ b/Engine/Plugins/Runtime/ReplicationGraph/Source/Private/ReplicationGraph.cpp
@@ -4325,7 +4325,7 @@ bool UReplicationGraphNode_ConnectionDormancyNode::IsNodeObsolete(uint32 Current
 
 // --------------------------------------------------------------------------------------------------------------------------------------------
 
-float UReplicationGraphNode_DormancyNode::MaxZForConnection = WORLD_MAX;
+FVector::FReal UReplicationGraphNode_DormancyNode::MaxZForConnection = WORLD_MAX;
 
 void UReplicationGraphNode_DormancyNode::CallFunctionOnValidConnectionNodes(FConnectionDormancyNodeFunction Function)
 {
diff --git a/Engine/Plugins/Runtime/ReplicationGraph/Source/Private/ReplicationGraphDebugging.cpp b/Engine/Plugins/Runtime/ReplicationGraph/Source/Private/ReplicationGraphDebugging.cpp
--- a/Engine/Plugins/Runtime/ReplicationGraph/Source/Private/ReplicationGraphDebugging.cpp
+++ b/Engine/Plugins/Runtime/ReplicationGraph/Source/Private/ReplicationGraphDebugging.cpp
@@ -622,7 +622,7 @@ FAutoConsoleCommandWithWorldAndArgs NetRepGraphSetDebugActorConnectionCmd(TEXT("
 		}
 		else
 		{
-			float ClosestMatchDistSq = WORLD_MAX;
+			FVector::FReal ClosestMatchDistSq = WORLD_MAX;
 			AActor* ClosestMatchActor = nullptr;
 			FVector CamLoc;
 			FRotator CamRot;
@@ -649,7 +649,7 @@ FAutoConsoleCommandWithWorldAndArgs NetRepGraphSetDebugActorConnectionCmd(TEXT("
 
 				if (Class)
 				{
-					float DistSq = (Actor->GetActorLocation() - CamLoc).SizeSquared2D();
+					FVector::FReal DistSq = (Actor->GetActorLocation() - CamLoc).SizeSquared2D();
 					if (DistSq < ClosestMatchDistSq)
 					{
 						ClosestMatchDistSq = DistSq;
diff --git a/Engine/Plugins/Runtime/ReplicationGraph/Source/Public/ReplicationGraph.h b/Engine/Plugins/Runtime/ReplicationGraph/Source/Public/ReplicationGraph.h
--- a/Engine/Plugins/Runtime/ReplicationGraph/Source/Public/ReplicationGraph.h
+++ b/Engine/Plugins/Runtime/ReplicationGraph/Source/Public/ReplicationGraph.h
@@ -507,7 +507,7 @@ class REPLICATIONGRAPH_API UReplicationGraphNode_DormancyNode : public UReplicat
 public:
 
 	/** Connection Z location has to be < this for ConnectionsNodes to be made. */
-	static float MaxZForConnection; 
+	static FVector::FReal MaxZForConnection; 
 
 public:
 
diff --git a/Engine/Source/Developer/BSPUtils/Private/BSPOps.cpp b/Engine/Source/Developer/BSPUtils/Private/BSPOps.cpp
--- a/Engine/Source/Developer/BSPUtils/Private/BSPOps.cpp
+++ b/Engine/Source/Developer/BSPUtils/Private/BSPOps.cpp
@@ -152,10 +152,10 @@ FPoly FBSPOps::BuildInfiniteFPoly( UModel* Model, int32 iNode )
 	EdPoly.Init();
 	EdPoly.Normal      = (FVector3f)Normal;
 	EdPoly.Base        = (FVector3f)Base;
-	new(EdPoly.Vertices) FVector3f(Base + Axis1*WORLD_MAX + Axis2*WORLD_MAX);
-	new(EdPoly.Vertices) FVector3f(Base - Axis1*WORLD_MAX + Axis2*WORLD_MAX);
-	new(EdPoly.Vertices) FVector3f(Base - Axis1*WORLD_MAX - Axis2*WORLD_MAX);
-	new(EdPoly.Vertices) FVector3f(Base + Axis1*WORLD_MAX - Axis2*WORLD_MAX);
+	new(EdPoly.Vertices) FVector3f(Base + Axis1*UE_OLD_WORLD_MAX + Axis2*UE_OLD_WORLD_MAX);
+	new(EdPoly.Vertices) FVector3f(Base - Axis1*UE_OLD_WORLD_MAX + Axis2*UE_OLD_WORLD_MAX);
+	new(EdPoly.Vertices) FVector3f(Base - Axis1*UE_OLD_WORLD_MAX - Axis2*UE_OLD_WORLD_MAX);
+	new(EdPoly.Vertices) FVector3f(Base + Axis1*UE_OLD_WORLD_MAX - Axis2*UE_OLD_WORLD_MAX);
 
 	return EdPoly;
 }
@@ -182,8 +182,8 @@ static void FilterBound
 	FVector3f&	Normal	= Model->Vectors[Surf.vNormal];
 	FBox		Bound(ForceInit);
 
-	Bound.Min.X = Bound.Min.Y = Bound.Min.Z = +WORLD_MAX;
-	Bound.Max.X = Bound.Max.Y = Bound.Max.Z = -WORLD_MAX;
+	Bound.Min.X = Bound.Min.Y = Bound.Min.Z = +UE_OLD_WORLD_MAX;	// LWC_TODO: Should be HALF_WORLD_MAX
+	Bound.Max.X = Bound.Max.Y = Bound.Max.Z = -UE_OLD_WORLD_MAX;
 
 	// Split bound into front half and back half.
 	FPoly** FrontList = new(FMemStack::Get(),nPolys*2+16)FPoly*; int32 nFront=0;
diff --git a/Engine/Source/Developer/MaterialUtilities/Private/MaterialUtilities.cpp b/Engine/Source/Developer/MaterialUtilities/Private/MaterialUtilities.cpp
--- a/Engine/Source/Developer/MaterialUtilities/Private/MaterialUtilities.cpp
+++ b/Engine/Source/Developer/MaterialUtilities/Private/MaterialUtilities.cpp
@@ -958,7 +958,7 @@ static bool ExportLandscapeMaterial(const ALandscapeProxy* InLandscape, const TS
 							FPlane(0,	0,	-1,	0),
 							FPlane(0,	0,	0,	1));
 				
-	const float ZOffset = WORLD_MAX;
+	const FMatrix::FReal ZOffset = UE_OLD_WORLD_MAX;
 	FMatrix ProjectionMatrix =  FReversedZOrthoMatrix(
 		LandscapeExtent.X,
 		LandscapeExtent.Y,
@@ -1481,7 +1481,7 @@ bool FMaterialUtilities::ExportBaseColor(ULandscapeComponent* LandscapeComponent
 		FPlane(0, 0, -1, 0),
 		FPlane(0, 0, 0, 1));
 
-	const float ZOffset = WORLD_MAX;
+	const FMatrix::FReal ZOffset = UE_OLD_WORLD_MAX;
 	FMatrix ProjectionMatrix = FReversedZOrthoMatrix(
 		LandscapeExtent.X,
 		LandscapeExtent.Y,
diff --git a/Engine/Source/Editor/DetailCustomizations/Private/AnimMontageSegmentDetails.cpp b/Engine/Source/Editor/DetailCustomizations/Private/AnimMontageSegmentDetails.cpp
--- a/Engine/Source/Editor/DetailCustomizations/Private/AnimMontageSegmentDetails.cpp
+++ b/Engine/Source/Editor/DetailCustomizations/Private/AnimMontageSegmentDetails.cpp
@@ -44,7 +44,7 @@ FAnimationSegmentViewportClient::FAnimationSegmentViewportClient(FPreviewScene&
 	DrawHelper.GridColorAxis = FColor(70, 70, 70);
 	DrawHelper.GridColorMajor = FColor(40, 40, 40);
 	DrawHelper.GridColorMinor =  FColor(20, 20, 20);
-	DrawHelper.PerspectiveGridSize = HALF_WORLD_MAX1;
+	DrawHelper.PerspectiveGridSize = UE_OLD_HALF_WORLD_MAX1;
 }
 
 
diff --git a/Engine/Source/Editor/MaterialEditor/Private/SMaterialEditorViewport.cpp b/Engine/Source/Editor/MaterialEditor/Private/SMaterialEditorViewport.cpp
--- a/Engine/Source/Editor/MaterialEditor/Private/SMaterialEditorViewport.cpp
+++ b/Engine/Source/Editor/MaterialEditor/Private/SMaterialEditorViewport.cpp
@@ -84,7 +84,7 @@ FMaterialEditorViewportClient::FMaterialEditorViewportClient(TWeakPtr<IMaterialE
 	DrawHelper.GridColorAxis = FColor(80,80,80);
 	DrawHelper.GridColorMajor = FColor(72,72,72);
 	DrawHelper.GridColorMinor = FColor(64,64,64);
-	DrawHelper.PerspectiveGridSize = HALF_WORLD_MAX1;
+	DrawHelper.PerspectiveGridSize = UE_OLD_HALF_WORLD_MAX1;
 	
 	SetViewMode(VMI_Lit);
 	
diff --git a/Engine/Source/Editor/Persona/Private/AnimationEditorViewportClient.cpp b/Engine/Source/Editor/Persona/Private/AnimationEditorViewportClient.cpp
--- a/Engine/Source/Editor/Persona/Private/AnimationEditorViewportClient.cpp
+++ b/Engine/Source/Editor/Persona/Private/AnimationEditorViewportClient.cpp
@@ -92,7 +92,7 @@ FAnimationViewportClient::FAnimationViewportClient(const TSharedRef<IPersonaPrev
 	check (ConfigOption);
 
 	// DrawHelper set up
-	DrawHelper.PerspectiveGridSize = HALF_WORLD_MAX1;
+	DrawHelper.PerspectiveGridSize = UE_OLD_HALF_WORLD_MAX1;
 	DrawHelper.AxesLineThickness = ConfigOption->bHighlightOrigin ? 1.0f : 0.0f;
 	DrawHelper.bDrawGrid = true;	// Toggling grid now relies on the show flag
 
diff --git a/Engine/Source/Editor/Persona/Private/SAnimationSequenceBrowser.cpp b/Engine/Source/Editor/Persona/Private/SAnimationSequenceBrowser.cpp
--- a/Engine/Source/Editor/Persona/Private/SAnimationSequenceBrowser.cpp
+++ b/Engine/Source/Editor/Persona/Private/SAnimationSequenceBrowser.cpp
@@ -1604,7 +1604,7 @@ FAnimationAssetViewportClient::FAnimationAssetViewportClient(FPreviewScene& InPr
 	DrawHelper.GridColorAxis = FColor(70, 70, 70);
 	DrawHelper.GridColorMajor = FColor(40, 40, 40);
 	DrawHelper.GridColorMinor = FColor(20, 20, 20);
-	DrawHelper.PerspectiveGridSize = HALF_WORLD_MAX1;
+	DrawHelper.PerspectiveGridSize = UE_OLD_HALF_WORLD_MAX1;
 	bDrawAxes = false;
 }
 
diff --git a/Engine/Source/Editor/UnrealEd/Private/EditorComponents.cpp b/Engine/Source/Editor/UnrealEd/Private/EditorComponents.cpp
--- a/Engine/Source/Editor/UnrealEd/Private/EditorComponents.cpp
+++ b/Engine/Source/Editor/UnrealEd/Private/EditorComponents.cpp
@@ -359,7 +359,7 @@ FEditorCommonDrawHelper::FEditorCommonDrawHelper()
 	, GridColorAxis(70, 70, 70)
 	, GridColorMajor(40, 40, 40)
 	, GridColorMinor(20, 20, 20)
-	, PerspectiveGridSize(HALF_WORLD_MAX1)
+	, PerspectiveGridSize(UE_OLD_HALF_WORLD_MAX1)
 	, PivotColor(FColor::Red)
 	, PivotSize(0.02f)
 	, NumCells(64)
@@ -424,15 +424,12 @@ void FEditorCommonDrawHelper::DrawBaseInfo(const FSceneView* View,FPrimitiveDraw
 // @todo UE - reimplement with new component attachment system
 }
 
-#define GRID_WORLD_MAX (2097152.0*0.5)
-
 void FEditorCommonDrawHelper::DrawOldGrid(const FSceneView* View,FPrimitiveDrawInterface* PDI)
 {
 	ESceneDepthPriorityGroup eDPG = (ESceneDepthPriorityGroup)DepthPriorityGroup;
 
 	bool bIsPerspective = ( View->ViewMatrices.GetProjectionMatrix().M[3][3] < 1.0f );
 
-	// LWC_TODO - replace with LWC version of HALF_WORLD_MAX once we have it
 	static double MaxGridExtent = 8.0 * 1024 * 1024 * 1024;
 
 	// Draw 3D perspective grid
diff --git a/Engine/Source/Editor/UnrealEd/Private/EditorViewportClient.cpp b/Engine/Source/Editor/UnrealEd/Private/EditorViewportClient.cpp
--- a/Engine/Source/Editor/UnrealEd/Private/EditorViewportClient.cpp
+++ b/Engine/Source/Editor/UnrealEd/Private/EditorViewportClient.cpp
@@ -1100,8 +1100,8 @@ FSceneView* FEditorViewportClient::CalcSceneView(FSceneViewFamily* ViewFamily, c
 		else
 		{
 			static_assert((bool)ERHIZBuffer::IsInverted, "Check all the Rotation Matrix transformations!");
-			float ZScale = 0.5f / HALF_WORLD_MAX;
-			float ZOffset = HALF_WORLD_MAX;
+			FMatrix::FReal ZScale = 0.5f / UE_OLD_HALF_WORLD_MAX;	// LWC_TODO: WORLD_MAX misuse?
+			FMatrix::FReal ZOffset = UE_OLD_HALF_WORLD_MAX;
 
 			//The divisor for the matrix needs to match the translation code.
 			const float Zoom = GetOrthoUnitsPerPixel(Viewport);
diff --git a/Engine/Source/Editor/UnrealEd/Private/GeomFitUtils.cpp b/Engine/Source/Editor/UnrealEd/Private/GeomFitUtils.cpp
--- a/Engine/Source/Editor/UnrealEd/Private/GeomFitUtils.cpp
+++ b/Engine/Source/Editor/UnrealEd/Private/GeomFitUtils.cpp
@@ -131,10 +131,10 @@ int32 GenerateKDopAsSimpleCollision(UStaticMesh* StaticMesh, const TArray<FVecto
 
 		Base = planes[i] * planes[i].W;
 
-		new(Polygon->Vertices) FVector3f(Base + AxisX * HALF_WORLD_MAX + AxisY * HALF_WORLD_MAX);
-		new(Polygon->Vertices) FVector3f(Base + AxisX * HALF_WORLD_MAX - AxisY * HALF_WORLD_MAX);
-		new(Polygon->Vertices) FVector3f(Base - AxisX * HALF_WORLD_MAX - AxisY * HALF_WORLD_MAX);
-		new(Polygon->Vertices) FVector3f(Base - AxisX * HALF_WORLD_MAX + AxisY * HALF_WORLD_MAX);
+		new(Polygon->Vertices) FVector3f(Base + AxisX * UE_OLD_HALF_WORLD_MAX + AxisY * UE_OLD_HALF_WORLD_MAX);
+		new(Polygon->Vertices) FVector3f(Base + AxisX * UE_OLD_HALF_WORLD_MAX - AxisY * UE_OLD_HALF_WORLD_MAX);
+		new(Polygon->Vertices) FVector3f(Base - AxisX * UE_OLD_HALF_WORLD_MAX - AxisY * UE_OLD_HALF_WORLD_MAX);
+		new(Polygon->Vertices) FVector3f(Base - AxisX * UE_OLD_HALF_WORLD_MAX + AxisY * UE_OLD_HALF_WORLD_MAX);
 
 		for(int32 j=0; j<planes.Num(); j++)
 		{
diff --git a/Engine/Source/Editor/UnrealEd/Private/LevelEditorViewport.cpp b/Engine/Source/Editor/UnrealEd/Private/LevelEditorViewport.cpp
--- a/Engine/Source/Editor/UnrealEd/Private/LevelEditorViewport.cpp
+++ b/Engine/Source/Editor/UnrealEd/Private/LevelEditorViewport.cpp
@@ -2304,7 +2304,7 @@ void FLevelEditorViewportClient::ProcessClick(FSceneView& View, HHitProxy* HitPr
 static float GPerspFrustumAngle=90.f;
 static float GPerspFrustumAspectRatio=1.77777f;
 static float GPerspFrustumStartDist=GNearClippingPlane;
-static float GPerspFrustumEndDist=HALF_WORLD_MAX;
+static float GPerspFrustumEndDist=UE_FLOAT_HUGE_DISTANCE;
 static FMatrix GPerspViewMatrix;
 
 
@@ -2331,7 +2331,7 @@ void FLevelEditorViewportClient::Tick(float DeltaTime)
 		GPerspFrustumAspectRatio=AspectRatio;
 		GPerspFrustumStartDist=GetNearClipPlane();
 
-		GPerspFrustumEndDist= HALF_WORLD_MAX;
+		GPerspFrustumEndDist= UE_FLOAT_HUGE_DISTANCE;
 
 		FSceneViewFamilyContext ViewFamily(FSceneViewFamily::ConstructionValues(
 			Viewport,
diff --git a/Engine/Source/Editor/UnrealEd/Private/SSkeletonWidget.cpp b/Engine/Source/Editor/UnrealEd/Private/SSkeletonWidget.cpp
--- a/Engine/Source/Editor/UnrealEd/Private/SSkeletonWidget.cpp
+++ b/Engine/Source/Editor/UnrealEd/Private/SSkeletonWidget.cpp
@@ -527,7 +527,7 @@ public:
 		DrawHelper.GridColorAxis = FColor(70, 70, 70);
 		DrawHelper.GridColorMajor = FColor(40, 40, 40);
 		DrawHelper.GridColorMinor =  FColor(20, 20, 20);
-		DrawHelper.PerspectiveGridSize = HALF_WORLD_MAX1;
+		DrawHelper.PerspectiveGridSize = UE_OLD_HALF_WORLD_MAX1;
 
 		bDisableInput = true;
 	}
diff --git a/Engine/Source/Editor/UnrealEd/Private/ThumbnailRendering/LevelThumbnailRenderer.cpp b/Engine/Source/Editor/UnrealEd/Private/ThumbnailRendering/LevelThumbnailRenderer.cpp
--- a/Engine/Source/Editor/UnrealEd/Private/ThumbnailRendering/LevelThumbnailRenderer.cpp
+++ b/Engine/Source/Editor/UnrealEd/Private/ThumbnailRendering/LevelThumbnailRenderer.cpp
@@ -72,7 +72,7 @@ FSceneView* ULevelThumbnailRenderer::CreateView(ULevel* Level, FSceneViewFamily*
 		FPlane(0,				0,				-1,		0),
 		FPlane(0,				0,				0,		1));
 
-	const float ZOffset = WORLD_MAX;
+	const FMatrix::FReal ZOffset = UE_OLD_WORLD_MAX;
 	ViewInitOptions.ProjectionMatrix =  FReversedZOrthoMatrix(
 		LevelBox.GetSize().X/2.f,
 		LevelBox.GetSize().Y/2.f,
diff --git a/Engine/Source/Editor/UnrealEd/Private/ThumbnailRendering/WorldThumbnailRenderer.cpp b/Engine/Source/Editor/UnrealEd/Private/ThumbnailRendering/WorldThumbnailRenderer.cpp
--- a/Engine/Source/Editor/UnrealEd/Private/ThumbnailRendering/WorldThumbnailRenderer.cpp
+++ b/Engine/Source/Editor/UnrealEd/Private/ThumbnailRendering/WorldThumbnailRenderer.cpp
@@ -256,13 +256,13 @@ FSceneView* UWorldThumbnailRenderer::CreateView(UWorld* World, FSceneViewFamily*
 		WorldSize2D.X = FMath::Abs(WorldSize2D.X);
 		WorldSize2D.Y = FMath::Abs(WorldSize2D.Y);
 		const bool bUseXAxis = (WorldSize2D.X / WorldSize2D.Y) > 1.f;
-		const float WorldAxisSize = bUseXAxis ? WorldSize2D.X : WorldSize2D.Y;
+		const FVector2D::FReal WorldAxisSize = bUseXAxis ? WorldSize2D.X : WorldSize2D.Y;
 		const uint32 ViewportAxisSize = bUseXAxis ? SizeX : SizeY;
-		const float OrthoZoom = WorldAxisSize / ViewportAxisSize / 2.f;
+		const FVector2D::FReal OrthoZoom = WorldAxisSize / ViewportAxisSize / 2.f;
 		const float OrthoWidth = FMath::Max(1.f, SizeX * OrthoZoom);
 		const float OrthoHeight = FMath::Max(1.f, SizeY * OrthoZoom);
 
-		const float ZOffset = HALF_WORLD_MAX;
+		const FMatrix::FReal ZOffset = UE_OLD_HALF_WORLD_MAX;
 		ProjectionMatrix = FReversedZOrthoMatrix(
 			OrthoWidth,
 			OrthoHeight,
diff --git a/Engine/Source/Editor/UnrealEd/Private/UnrealEdSrv.cpp b/Engine/Source/Editor/UnrealEd/Private/UnrealEdSrv.cpp
--- a/Engine/Source/Editor/UnrealEd/Private/UnrealEdSrv.cpp
+++ b/Engine/Source/Editor/UnrealEd/Private/UnrealEdSrv.cpp
@@ -1827,9 +1827,9 @@ FPoly* CreateHugeTrianglePolygonOnPlane( const FPlane* InPlane )
 	FPoly* Triangle = new FPoly();
 
 	FVector Center = FVector( InPlane->X, InPlane->Y, InPlane->Z ) * InPlane->W;
-	FVector V0 = Center + (A * WORLD_MAX);
-	FVector V1 = Center + (B * WORLD_MAX);
-	FVector V2 = Center - (((A + B) / 2.0f) * WORLD_MAX);
+	FVector V0 = Center + (A * UE_OLD_WORLD_MAX);	// LWC_TODO: WORLD_MAX misuse?
+	FVector V1 = Center + (B * UE_OLD_WORLD_MAX);
+	FVector V2 = Center - (((A + B) / 2.0f) * UE_OLD_WORLD_MAX);
 
 	// Create a triangle that lays on InPlane
 
diff --git a/Engine/Source/Editor/WorldBrowser/Private/LevelCollectionModel.cpp b/Engine/Source/Editor/WorldBrowser/Private/LevelCollectionModel.cpp
--- a/Engine/Source/Editor/WorldBrowser/Private/LevelCollectionModel.cpp
+++ b/Engine/Source/Editor/WorldBrowser/Private/LevelCollectionModel.cpp
@@ -1150,14 +1150,14 @@ void FLevelCollectionModel::BroadcastPostLevelsUnloaded()
 	PostLevelsUnloaded.Broadcast();
 }
 
-float FLevelCollectionModel::EditableAxisLength()
+double FLevelCollectionModel::EditableAxisLength()
 { 
 	return HALF_WORLD_MAX; 
 };
 
 FBox FLevelCollectionModel::EditableWorldArea()
 {
-	float AxisLength = EditableAxisLength();
+	FVector::FReal AxisLength = EditableAxisLength();
 	
 	return FBox(	
 		FVector(-AxisLength, -AxisLength, -AxisLength), 
diff --git a/Engine/Source/Editor/WorldBrowser/Private/LevelCollectionModel.h b/Engine/Source/Editor/WorldBrowser/Private/LevelCollectionModel.h
--- a/Engine/Source/Editor/WorldBrowser/Private/LevelCollectionModel.h
+++ b/Engine/Source/Editor/WorldBrowser/Private/LevelCollectionModel.h
@@ -296,7 +296,7 @@ public:
 	void BroadcastPostLevelsUnloaded();
 	
 	/** Editable world axis length  */
-	static float EditableAxisLength();
+	static double EditableAxisLength();
 
 	/** Editable world bounds */
 	static FBox EditableWorldArea();
diff --git a/Engine/Source/Runtime/Engine/Classes/Camera/CameraTypes.h b/Engine/Source/Runtime/Engine/Classes/Camera/CameraTypes.h
--- a/Engine/Source/Runtime/Engine/Classes/Camera/CameraTypes.h
+++ b/Engine/Source/Runtime/Engine/Classes/Camera/CameraTypes.h
@@ -121,7 +121,7 @@ struct FMinimalViewInfo
 		, DesiredFOV(90.0f)
 		, OrthoWidth(512.0f)
 		, OrthoNearClipPlane(0.0f)
-		, OrthoFarClipPlane(WORLD_MAX)
+		, OrthoFarClipPlane(UE_OLD_WORLD_MAX)
 		, PerspectiveNearClipPlane(-1.0f)
 		, AspectRatio(1.33333333f)
 		, bConstrainAspectRatio(false)
diff --git a/Engine/Source/Runtime/Engine/Private/Camera/CameraComponent.cpp b/Engine/Source/Runtime/Engine/Private/Camera/CameraComponent.cpp
--- a/Engine/Source/Runtime/Engine/Private/Camera/CameraComponent.cpp
+++ b/Engine/Source/Runtime/Engine/Private/Camera/CameraComponent.cpp
@@ -45,7 +45,7 @@ UCameraComponent::UCameraComponent(const FObjectInitializer& ObjectInitializer)
 	AspectRatio = 1.777778f;
 	OrthoWidth = 512.0f;
 	OrthoNearClipPlane = 0.0f;
-	OrthoFarClipPlane = WORLD_MAX;
+	OrthoFarClipPlane = UE_OLD_WORLD_MAX;
 	bConstrainAspectRatio = false;
 	bUseFieldOfViewForLOD = true;
 	PostProcessBlendWeight = 1.0f;
diff --git a/Engine/Source/Runtime/Engine/Private/Collision/KAggregateGeom.cpp b/Engine/Source/Runtime/Engine/Private/Collision/KAggregateGeom.cpp
--- a/Engine/Source/Runtime/Engine/Private/Collision/KAggregateGeom.cpp
+++ b/Engine/Source/Runtime/Engine/Private/Collision/KAggregateGeom.cpp
@@ -759,10 +759,10 @@ bool FKConvexElem::HullFromPlanes(const TArray<FPlane>& InPlanes, const TArray<F
 
 		const FVector3f Base = FVector3f(InPlanes[i] * InPlanes[i].W);
 
-		new(Polygon.Vertices) FVector3f(Base + AxisX * HALF_WORLD_MAX + AxisY * HALF_WORLD_MAX);
-		new(Polygon.Vertices) FVector3f(Base - AxisX * HALF_WORLD_MAX + AxisY * HALF_WORLD_MAX);
-		new(Polygon.Vertices) FVector3f(Base - AxisX * HALF_WORLD_MAX - AxisY * HALF_WORLD_MAX);
-		new(Polygon.Vertices) FVector3f(Base + AxisX * HALF_WORLD_MAX - AxisY * HALF_WORLD_MAX);
+		new(Polygon.Vertices) FVector3f(Base + AxisX * UE_OLD_HALF_WORLD_MAX + AxisY * UE_OLD_HALF_WORLD_MAX);
+		new(Polygon.Vertices) FVector3f(Base - AxisX * UE_OLD_HALF_WORLD_MAX + AxisY * UE_OLD_HALF_WORLD_MAX);
+		new(Polygon.Vertices) FVector3f(Base - AxisX * UE_OLD_HALF_WORLD_MAX - AxisY * UE_OLD_HALF_WORLD_MAX);
+		new(Polygon.Vertices) FVector3f(Base + AxisX * UE_OLD_HALF_WORLD_MAX - AxisY * UE_OLD_HALF_WORLD_MAX);
 
 		for(int32 j=0; j<InPlanes.Num(); j++)
 		{
diff --git a/Engine/Source/Runtime/Engine/Private/Components/DirectionalLightComponent.cpp b/Engine/Source/Runtime/Engine/Private/Components/DirectionalLightComponent.cpp
--- a/Engine/Source/Runtime/Engine/Private/Components/DirectionalLightComponent.cpp
+++ b/Engine/Source/Runtime/Engine/Private/Components/DirectionalLightComponent.cpp
@@ -47,9 +47,9 @@ static TAutoConsoleVariable<float> CVarPerObjectCastDistanceRadiusScale(
 	);
 static TAutoConsoleVariable<float> CVarPerObjectCastDistanceMin(
 	TEXT("r.Shadow.PerObjectCastDistanceMin"),
-	(float)HALF_WORLD_MAX / 32.0f,
+	(float)UE_FLOAT_HUGE_DISTANCE / 32.0f,
 	TEXT("Minimum cast distance for Per-Object shadows, i.e., CastDistDance = Max(r.Shadow.PerObjectCastDistanceRadiusScale * object-radius, r.Shadow.PerObjectCastDistanceMin).\n")
-	TEXT("  Default: HALF_WORLD_MAX / 32.0f"),
+	TEXT("  Default: UE_FLOAT_HUGE_DISTANCE / 32.0f"),
 	ECVF_RenderThreadSafe
 	);
 static TAutoConsoleVariable<int32> CVarMaxNumFarShadowCascades(
diff --git a/Engine/Source/Runtime/Engine/Private/Components/MaterialBillboardComponent.cpp b/Engine/Source/Runtime/Engine/Private/Components/MaterialBillboardComponent.cpp
--- a/Engine/Source/Runtime/Engine/Private/Components/MaterialBillboardComponent.cpp
+++ b/Engine/Source/Runtime/Engine/Private/Components/MaterialBillboardComponent.cpp
@@ -298,19 +298,19 @@ FPrimitiveSceneProxy* UMaterialBillboardComponent::CreateSceneProxy()
 
 FBoxSphereBounds UMaterialBillboardComponent::CalcBounds(const FTransform& LocalToWorld) const
 {
-	float BoundsSize = 1.0f;
+	FTransform::FReal BoundsSize = 1.0f;
 	for (int32 i = 0; i < Elements.Num(); ++i)
 	{
 		if (Elements[i].bSizeIsInScreenSpace)
 		{
 			// Workaround static bounds by disabling culling. Still allows override such as 'use parent bounds', etc.
 			// Note: Bounds are dynamically calculated at draw time per view, so difficult to cull correctly. (UE-4725)
-			BoundsSize = float(HALF_WORLD_MAX); 
+			BoundsSize = HALF_WORLD_MAX; 
 			break;
 		}
 		else
 		{
-			BoundsSize = FMath::Max3(BoundsSize, Elements[i].BaseSizeX, Elements[i].BaseSizeY);
+			BoundsSize = FMath::Max3<FTransform::FReal>(BoundsSize, Elements[i].BaseSizeX, Elements[i].BaseSizeY);
 		}
 	}
 	BoundsSize *= LocalToWorld.GetMaximumAxisScale();
diff --git a/Engine/Source/Runtime/Engine/Private/Particles/ParticleComponents.cpp b/Engine/Source/Runtime/Engine/Private/Particles/ParticleComponents.cpp
--- a/Engine/Source/Runtime/Engine/Private/Particles/ParticleComponents.cpp
+++ b/Engine/Source/Runtime/Engine/Private/Particles/ParticleComponents.cpp
@@ -7086,10 +7086,10 @@ int32 UParticleSystemComponent::DetermineLODLevelForLocation(const FVector& Effe
 		}
 
 		// This will now put everything in LODLevel 0 (high detail) by default
-		float LODDistanceSqr = (PlayerViewLocations.Num() ? FMath::Square(WORLD_MAX) : 0.0f);
+		FVector::FReal LODDistanceSqr = (PlayerViewLocations.Num() ? FMath::Square(WORLD_MAX) : 0.0f);
 		for (const FVector& ViewLocation : PlayerViewLocations)
 		{
-			const float DistanceToEffectSqr = FVector(ViewLocation - EffectLocation).SizeSquared();
+			const FVector::FReal DistanceToEffectSqr = FVector(ViewLocation - EffectLocation).SizeSquared();
 			if (DistanceToEffectSqr < LODDistanceSqr)
 			{
 				LODDistanceSqr = DistanceToEffectSqr;
diff --git a/Engine/Source/Runtime/Engine/Private/Particles/ParticleEmitterInstances.cpp b/Engine/Source/Runtime/Engine/Private/Particles/ParticleEmitterInstances.cpp
--- a/Engine/Source/Runtime/Engine/Private/Particles/ParticleEmitterInstances.cpp
+++ b/Engine/Source/Runtime/Engine/Private/Particles/ParticleEmitterInstances.cpp
@@ -1353,7 +1353,7 @@ void FParticleEmitterInstance::UpdateBoundingBox(float DeltaTime)
 				NewRotation = Particle.Rotation;
 			}
 
-			float LocalMax(0.0f);
+			FVector::FReal LocalMax(0.0f);
 
 			if (bUpdateBox)
 			{	
@@ -1390,12 +1390,12 @@ void FParticleEmitterInstance::UpdateBoundingBox(float DeltaTime)
 
 				// Treat each particle as a cube whose sides are the length of the maximum component
 				// This handles the particle's extents changing due to being camera facing
-				MinVal[0] = FMath::Min<float>(MinVal[0], PositionForBounds.X - LocalMax);
-				MaxVal[0] = FMath::Max<float>(MaxVal[0], PositionForBounds.X + LocalMax);
-				MinVal[1] = FMath::Min<float>(MinVal[1], PositionForBounds.Y - LocalMax);
-				MaxVal[1] = FMath::Max<float>(MaxVal[1], PositionForBounds.Y + LocalMax);
-				MinVal[2] = FMath::Min<float>(MinVal[2], PositionForBounds.Z - LocalMax);
-				MaxVal[2] = FMath::Max<float>(MaxVal[2], PositionForBounds.Z + LocalMax);
+				MinVal[0] = FMath::Min(MinVal[0], PositionForBounds.X - LocalMax);
+				MaxVal[0] = FMath::Max(MaxVal[0], PositionForBounds.X + LocalMax);
+				MinVal[1] = FMath::Min(MinVal[1], PositionForBounds.Y - LocalMax);
+				MaxVal[1] = FMath::Max(MaxVal[1], PositionForBounds.Y + LocalMax);
+				MinVal[2] = FMath::Min(MinVal[2], PositionForBounds.Z - LocalMax);
+				MaxVal[2] = FMath::Max(MaxVal[2], PositionForBounds.Z + LocalMax);
 			}
 		}
 
@@ -1439,7 +1439,7 @@ void FParticleEmitterInstance::ForceUpdateBoundingBox()
 		{
 			DECLARE_PARTICLE(Particle, ParticleData + ParticleStride * ParticleIndices[i]);
 
-			float LocalMax(0.0f);
+			FVector::FReal LocalMax(0.0f);
 
 			if (OrbitOffsetValue == -1)
 			{
@@ -1463,12 +1463,12 @@ void FParticleEmitterInstance::ForceUpdateBoundingBox()
 
 			// Treat each particle as a cube whose sides are the length of the maximum component
 			// This handles the particle's extents changing due to being camera facing
-			MinVal[0] = FMath::Min<float>(MinVal[0], PositionForBounds.X - LocalMax);
-			MaxVal[0] = FMath::Max<float>(MaxVal[0], PositionForBounds.X + LocalMax);
-			MinVal[1] = FMath::Min<float>(MinVal[1], PositionForBounds.Y - LocalMax);
-			MaxVal[1] = FMath::Max<float>(MaxVal[1], PositionForBounds.Y + LocalMax);
-			MinVal[2] = FMath::Min<float>(MinVal[2], PositionForBounds.Z - LocalMax);
-			MaxVal[2] = FMath::Max<float>(MaxVal[2], PositionForBounds.Z + LocalMax);
+			MinVal[0] = FMath::Min(MinVal[0], PositionForBounds.X - LocalMax);
+			MaxVal[0] = FMath::Max(MaxVal[0], PositionForBounds.X + LocalMax);
+			MinVal[1] = FMath::Min(MinVal[1], PositionForBounds.Y - LocalMax);
+			MaxVal[1] = FMath::Max(MaxVal[1], PositionForBounds.Y + LocalMax);
+			MinVal[2] = FMath::Min(MinVal[2], PositionForBounds.Z - LocalMax);
+			MaxVal[2] = FMath::Max(MaxVal[2], PositionForBounds.Z + LocalMax);
 		}
 
 		ParticleBoundingBox = FBox(MinVal, MaxVal);
@@ -3548,12 +3548,12 @@ void FParticleMeshEmitterInstance::UpdateBoundingBox(float DeltaTime)
 					PositionForBounds = ComponentToWorld.TransformPosition(NewLocation);
 				}
 
-				MinVal[0] = FMath::Min<float>(MinVal[0], PositionForBounds.X - LocalExtent.X);
-				MaxVal[0] = FMath::Max<float>(MaxVal[0], PositionForBounds.X + LocalExtent.X);
-				MinVal[1] = FMath::Min<float>(MinVal[1], PositionForBounds.Y - LocalExtent.Y);
-				MaxVal[1] = FMath::Max<float>(MaxVal[1], PositionForBounds.Y + LocalExtent.Y);
-				MinVal[2] = FMath::Min<float>(MinVal[2], PositionForBounds.Z - LocalExtent.Z);
-				MaxVal[2] = FMath::Max<float>(MaxVal[2], PositionForBounds.Z + LocalExtent.Z);
+				MinVal[0] = FMath::Min(MinVal[0], PositionForBounds.X - LocalExtent.X);
+				MaxVal[0] = FMath::Max(MaxVal[0], PositionForBounds.X + LocalExtent.X);
+				MinVal[1] = FMath::Min(MinVal[1], PositionForBounds.Y - LocalExtent.Y);
+				MaxVal[1] = FMath::Max(MaxVal[1], PositionForBounds.Y + LocalExtent.Y);
+				MinVal[2] = FMath::Min(MinVal[2], PositionForBounds.Z - LocalExtent.Z);
+				MaxVal[2] = FMath::Max(MaxVal[2], PositionForBounds.Z + LocalExtent.Z);
 			}
 		}
 
diff --git a/Engine/Source/Runtime/Engine/Private/Particles/ParticleModules_Collision.cpp b/Engine/Source/Runtime/Engine/Private/Particles/ParticleModules_Collision.cpp
--- a/Engine/Source/Runtime/Engine/Private/Particles/ParticleModules_Collision.cpp
+++ b/Engine/Source/Runtime/Engine/Private/Particles/ParticleModules_Collision.cpp
@@ -225,7 +225,7 @@ void UParticleModuleCollision::Update(FParticleEmitterInstance* Owner, int32 Off
 		else
 		{
 			// If the MaxCollisionDistance is greater than WORLD_MAX, they obviously want the check disabled...
-			if (MaxCollisionDistance < WORLD_MAX)
+			if (MaxCollisionDistance < UE_OLD_WORLD_MAX)	// LWC_TODO: Fix needed to handle WORLD_MAX increase. Remove?
 			{
 				// If we have at least a few particles, do a simple check vs. the bounds
 				if (Owner->ActiveParticles > 7)
@@ -255,7 +255,7 @@ void UParticleModuleCollision::Update(FParticleEmitterInstance* Owner, int32 Off
 							// bounds rather than shorten the distance checked as it is usually used for.
 							float InvDistanceFactor = 1.0f / PlayerLODDistanceFactor[PlyrIdx];
 							FBox CheckBounds = BoundingBox;
-							float BoxExpansionValue = MaxCollisionDistance * InvDistanceFactor;
+							FBox::FReal BoxExpansionValue = MaxCollisionDistance * InvDistanceFactor;
 							BoxExpansionValue += BoxExpansionValue * 0.075f;
 							// Expand it by the max collision distance (and a little bit extra)
 							CheckBounds = CheckBounds.ExpandBy(BoxExpansionValue);
@@ -302,7 +302,7 @@ void UParticleModuleCollision::Update(FParticleEmitterInstance* Owner, int32 Off
 		}
 	}
 
-	float SquaredMaxCollisionDistance = FMath::Square(MaxCollisionDistance);
+	FVector::FReal SquaredMaxCollisionDistance = FMath::Square<FVector::FReal>(MaxCollisionDistance);
 	BEGIN_UPDATE_LOOP;
 	{
 		if ((Particle.Flags & STATE_Particle_CollisionIgnoreCheck) != 0)
@@ -359,7 +359,7 @@ void UParticleModuleCollision::Update(FParticleEmitterInstance* Owner, int32 Off
 
 		FVector End = Location + Direction * Size / DirScalar;
 
-		if ((World->IsGameWorld() == true) && (MaxCollisionDistance < WORLD_MAX))
+		if ((World->IsGameWorld() == true) && (MaxCollisionDistance < UE_OLD_WORLD_MAX))	// LWC_TODO: Fix needed to handle WORLD_MAX increase.
 		{
 			// LOD collision by distance
 			bool bCloseEnough = false;
diff --git a/Engine/Source/Runtime/Engine/Private/Particles/ParticleSystemRender.cpp b/Engine/Source/Runtime/Engine/Private/Particles/ParticleSystemRender.cpp
--- a/Engine/Source/Runtime/Engine/Private/Particles/ParticleSystemRender.cpp
+++ b/Engine/Source/Runtime/Engine/Private/Particles/ParticleSystemRender.cpp
@@ -6765,11 +6765,11 @@ void FParticleSystemSceneProxy::GetObjectPositionAndScale(const FSceneView& View
 		// Scales to transform the view space positions corresponding to SystemPositionForMacroUVs +- SystemRadiusForMacroUVs into [0, 1] in xy
 		// Scales to transform the screen space positions corresponding to SystemPositionForMacroUVs +- SystemRadiusForMacroUVs into [0, 1] in zw
 
-		const float RightNDCPosX = (float)(RightPostProjectionPosition.X / RightPostProjectionPosition.W); // LWC_TODO: precision loss
-		const float UpNDCPosY = (float)(UpPostProjectionPosition.Y / UpPostProjectionPosition.W); // LWC_TODO: precision loss
-		float DX = FMath::Min<float>(RightNDCPosX - ObjectNDCPosition.X, WORLD_MAX);
-		float DY = FMath::Min<float>(UpNDCPosY - ObjectNDCPosition.Y, WORLD_MAX);
-		if (DX != 0.0f && DY != 0.0f && !FMath::IsNaN(DX) && FMath::IsFinite(DX) && !FMath::IsNaN(DY) && FMath::IsFinite(DY))
+		const FVector4::FReal RightNDCPosX = RightPostProjectionPosition.X / RightPostProjectionPosition.W;
+		const FVector4::FReal UpNDCPosY = UpPostProjectionPosition.Y / UpPostProjectionPosition.W;
+		FVector4::FReal DX = FMath::Min(RightNDCPosX - ObjectNDCPosition.X, WORLD_MAX);
+		FVector4::FReal DY = FMath::Min(UpNDCPosY - ObjectNDCPosition.Y, WORLD_MAX);
+		if (DX != 0 && DY != 0 && !FMath::IsNaN(DX) && FMath::IsFinite(DX) && !FMath::IsNaN(DY) && FMath::IsFinite(DY))
 		{
 			ObjectMacroUVScales = FVector2D(1.0f / DX, -1.0f / DY);
 		}
diff --git a/Engine/Source/Runtime/Engine/Private/Polygon.cpp b/Engine/Source/Runtime/Engine/Private/Polygon.cpp
--- a/Engine/Source/Runtime/Engine/Private/Polygon.cpp
+++ b/Engine/Source/Runtime/Engine/Private/Polygon.cpp
@@ -834,10 +834,10 @@ FPoly FPoly::BuildInfiniteFPoly(const FPlane& InPlane)
 	EdPoly.Normal.Y    = InPlane.Y;
 	EdPoly.Normal.Z    = InPlane.Z;
 	EdPoly.Base        = EdPoly.Normal * InPlane.W;
-	EdPoly.Vertices.Add( EdPoly.Base + (FVector3f)Axis1*HALF_WORLD_MAX + (FVector3f)Axis2*HALF_WORLD_MAX );
-	EdPoly.Vertices.Add( EdPoly.Base - (FVector3f)Axis1*HALF_WORLD_MAX + (FVector3f)Axis2*HALF_WORLD_MAX );
-	EdPoly.Vertices.Add( EdPoly.Base - (FVector3f)Axis1*HALF_WORLD_MAX - (FVector3f)Axis2*HALF_WORLD_MAX );
-	EdPoly.Vertices.Add( EdPoly.Base + (FVector3f)Axis1*HALF_WORLD_MAX - (FVector3f)Axis2*HALF_WORLD_MAX );
+	EdPoly.Vertices.Add( EdPoly.Base + (FVector3f)Axis1*UE_OLD_HALF_WORLD_MAX + (FVector3f)Axis2*UE_OLD_HALF_WORLD_MAX );
+	EdPoly.Vertices.Add( EdPoly.Base - (FVector3f)Axis1*UE_OLD_HALF_WORLD_MAX + (FVector3f)Axis2*UE_OLD_HALF_WORLD_MAX );
+	EdPoly.Vertices.Add( EdPoly.Base - (FVector3f)Axis1*UE_OLD_HALF_WORLD_MAX - (FVector3f)Axis2*UE_OLD_HALF_WORLD_MAX );
+	EdPoly.Vertices.Add( EdPoly.Base + (FVector3f)Axis1*UE_OLD_HALF_WORLD_MAX - (FVector3f)Axis2*UE_OLD_HALF_WORLD_MAX );
 
 	return EdPoly;
 }
diff --git a/Engine/Source/Runtime/Engine/Private/WorldComposition.cpp b/Engine/Source/Runtime/Engine/Private/WorldComposition.cpp
--- a/Engine/Source/Runtime/Engine/Private/WorldComposition.cpp
+++ b/Engine/Source/Runtime/Engine/Private/WorldComposition.cpp
@@ -37,7 +37,7 @@ UWorldComposition::UWorldComposition(const FObjectInitializer& ObjectInitializer
 	, TilesStreamingTimeThreshold(1.0)
 	, bLoadAllTilesDuringCinematic(false)
 	, bRebaseOriginIn3DSpace(false)
-	, RebaseOriginDistance(HALF_WORLD_MAX1*0.5f)
+	, RebaseOriginDistance(UE_OLD_HALF_WORLD_MAX1*0.5f)
 {
 }
 
diff --git a/Engine/Source/Runtime/Engine/Private/WorldSettings.cpp b/Engine/Source/Runtime/Engine/Private/WorldSettings.cpp
--- a/Engine/Source/Runtime/Engine/Private/WorldSettings.cpp
+++ b/Engine/Source/Runtime/Engine/Private/WorldSettings.cpp
@@ -86,7 +86,7 @@ PRAGMA_ENABLE_DEPRECATION_WARNINGS
 	bGenerateSingleClusterForLevel = false;
 #endif
 
-	KillZ = -HALF_WORLD_MAX1;
+	KillZ = -UE_OLD_HALF_WORLD_MAX1;	// LWC_TODO: HALF_WORLD_MAX1? Something else?
 	KillZDamageType = ConstructorStatics.DmgType_Environmental_Object.Object;
 
 #if WITH_EDITORONLY_DATA
diff --git a/Engine/Source/Runtime/Landscape/Private/LandscapeEdit.cpp b/Engine/Source/Runtime/Landscape/Private/LandscapeEdit.cpp
--- a/Engine/Source/Runtime/Landscape/Private/LandscapeEdit.cpp
+++ b/Engine/Source/Runtime/Landscape/Private/LandscapeEdit.cpp
@@ -3728,7 +3728,7 @@ bool ULandscapeInfo::GetSelectedExtent(int32& MinX, int32& MinY, int32& MaxX, in
 FVector ULandscapeInfo::GetLandscapeCenterPos(float& LengthZ, int32 MinX /*= MAX_INT*/, int32 MinY /*= MAX_INT*/, int32 MaxX /*= MIN_INT*/, int32 MaxY /*= MIN_INT*/)
 {
 	// MinZ, MaxZ is Local coordinate
-	float MaxZ = -HALF_WORLD_MAX, MinZ = HALF_WORLD_MAX;
+	float MaxZ = -UE_OLD_HALF_WORLD_MAX, MinZ = UE_OLD_HALF_WORLD_MAX;
 	const float ScaleZ = DrawScale.Z;
 
 	if (MinX == MAX_int32)
diff --git a/Engine/Source/Runtime/Landscape/Private/LandscapeGizmoActor.cpp b/Engine/Source/Runtime/Landscape/Private/LandscapeGizmoActor.cpp
--- a/Engine/Source/Runtime/Landscape/Private/LandscapeGizmoActor.cpp
+++ b/Engine/Source/Runtime/Landscape/Private/LandscapeGizmoActor.cpp
@@ -741,7 +741,7 @@ void ALandscapeGizmoActiveActor::FitMinMaxHeight()
 {
 	if (TargetLandscapeInfo)
 	{
-		float MinZ = HALF_WORLD_MAX, MaxZ = -HALF_WORLD_MAX;
+		FVector::FReal MinZ = UE_OLD_HALF_WORLD_MAX, MaxZ = -UE_OLD_HALF_WORLD_MAX;
 		// Change MinRelativeZ and RelativeZScale to fit Gizmo Box
 		for (auto It = SelectedData.CreateConstIterator(); It; ++It )
 		{
@@ -750,7 +750,7 @@ void ALandscapeGizmoActiveActor::FitMinMaxHeight()
 			MaxZ = FMath::Max(MaxZ, Data.HeightData);
 		}
 
-		if (MinZ != HALF_WORLD_MAX && MaxZ > MinZ + KINDA_SMALL_NUMBER)
+		if (MinZ != UE_OLD_HALF_WORLD_MAX && MaxZ > MinZ + KINDA_SMALL_NUMBER)
 		{
 			MinRelativeZ = MinZ;
 			RelativeScaleZ = 1.f / (MaxZ - MinZ);
diff --git a/Engine/Source/Runtime/Landscape/Private/LandscapeGrass.cpp b/Engine/Source/Runtime/Landscape/Private/LandscapeGrass.cpp
--- a/Engine/Source/Runtime/Landscape/Private/LandscapeGrass.cpp
+++ b/Engine/Source/Runtime/Landscape/Private/LandscapeGrass.cpp
@@ -675,7 +675,7 @@ public:
 		                              FPlane(0, 0,-1, 0),
 		                              FPlane(0, 0, 0, 1));
 
-		const float ZOffset = WORLD_MAX;
+		const FMatrix::FReal ZOffset = UE_OLD_WORLD_MAX;
 		ProjectionMatrix = FReversedZOrthoMatrix(
 			TargetExtent.X,
 			TargetExtent.Y,
diff --git a/Engine/Source/Runtime/Landscape/Private/LandscapePhysicalMaterial.cpp b/Engine/Source/Runtime/Landscape/Private/LandscapePhysicalMaterial.cpp
--- a/Engine/Source/Runtime/Landscape/Private/LandscapePhysicalMaterial.cpp
+++ b/Engine/Source/Runtime/Landscape/Private/LandscapePhysicalMaterial.cpp
@@ -424,7 +424,7 @@ namespace
 				const FVector TargetCenter = ComponentTransform.TransformPosition(FVector(TargetSizeMinusOne, 0.f) * 0.5f);
 				const FVector TargetExtent = FVector(TargetSize, 0.0f) * ComponentTransform.GetScale3D() * 0.5f;
 				const FMatrix ViewRotationMatrix = FInverseRotationMatrix(ComponentTransform.Rotator()) * FMatrix(FPlane(1, 0, 0, 0), FPlane(0, -1, 0, 0), FPlane(0, 0, -1, 0), FPlane(0, 0, 0, 1));
-				const float ZOffset = WORLD_MAX;
+				const FMatrix::FReal ZOffset = UE_OLD_WORLD_MAX;
 				const FMatrix ProjectionMatrix = FReversedZOrthoMatrix(TargetExtent.X, TargetExtent.Y, 0.5f / ZOffset, ZOffset);
 
 				if (Task.TargetSize != TargetSize)
diff --git a/Engine/Source/Runtime/Renderer/Private/Renderer.cpp b/Engine/Source/Runtime/Renderer/Private/Renderer.cpp
--- a/Engine/Source/Runtime/Renderer/Private/Renderer.cpp
+++ b/Engine/Source/Runtime/Renderer/Private/Renderer.cpp
@@ -437,7 +437,7 @@ void FRendererModule::GPUBenchmark(FSynthBenchmarkResults& InOut, float WorkScal
 	FSceneViewInitOptions ViewInitOptions;
 	FIntRect ViewRect(0, 0, 1, 1);
 
-	FBox LevelBox(FVector(-WORLD_MAX), FVector(+WORLD_MAX));
+	FBox LevelBox(FVector(-UE_OLD_WORLD_MAX), FVector(+UE_OLD_WORLD_MAX));	// LWC_TODO: Scale to renderable world bounds?
 	ViewInitOptions.SetViewRectangle(ViewRect);
 
 	// Initialize Projection Matrix and ViewMatrix since FSceneView initialization is doing some math on them.
@@ -450,7 +450,7 @@ void FRendererModule::GPUBenchmark(FSynthBenchmarkResults& InOut, float WorkScal
 		FPlane(0, 0, -1, 0),
 		FPlane(0, 0, 0, 1));
 
-	const float ZOffset = WORLD_MAX;
+	const FVector::FReal ZOffset = UE_OLD_WORLD_MAX;
 	ViewInitOptions.ProjectionMatrix = FReversedZOrthoMatrix(
 		LevelBox.GetSize().X / 2.f,
 		LevelBox.GetSize().Y / 2.f,
diff --git a/Engine/Source/Runtime/Renderer/Private/SceneCaptureRendering.cpp b/Engine/Source/Runtime/Renderer/Private/SceneCaptureRendering.cpp
--- a/Engine/Source/Runtime/Renderer/Private/SceneCaptureRendering.cpp
+++ b/Engine/Source/Runtime/Renderer/Private/SceneCaptureRendering.cpp
@@ -654,7 +654,7 @@ static void BuildOrthoMatrix(FIntPoint InRenderTargetSize, float InOrthoWidth, i
 	const float OrthoHeight = InOrthoWidth / 2.0f * XAxisMultiplier / YAxisMultiplier;
 
 	const float NearPlane = 0;
-	const float FarPlane = WORLD_MAX / 8.0f;
+	const float FarPlane = UE_FLOAT_HUGE_DISTANCE / 4.0f;
 
 	const float ZScale = 1.0f / (FarPlane - NearPlane);
 	const float ZOffset = -NearPlane;
diff --git a/Engine/Source/Runtime/Renderer/Private/ShadowSetup.cpp b/Engine/Source/Runtime/Renderer/Private/ShadowSetup.cpp
--- a/Engine/Source/Runtime/Renderer/Private/ShadowSetup.cpp
+++ b/Engine/Source/Runtime/Renderer/Private/ShadowSetup.cpp
@@ -707,7 +707,7 @@ FProjectedShadowInfo::FProjectedShadowInfo()
 	, bNaniteGeometry(true)
 	, bIncludeInScreenSpaceShadowMask(true)
 	, bContainsNaniteSubjects(false)
-	, PerObjectShadowFadeStart(WORLD_MAX)
+	, PerObjectShadowFadeStart(UE_OLD_WORLD_MAX)		// LWC_TODO: Upgrade to double? Should have been HALF_WORLD_MAX to match precision.
 	, InvPerObjectShadowFadeLength(0.0f)
 	, OverlappedUVOnCachedShadowMap(-1.0f, -1.0f, -1.0f, -1.0f)
 	, OverlappedUVOnCurrentShadowMap(-1.0f, -1.0f, -1.0f, -1.0f)
diff --git a/Engine/Source/Runtime/Renderer/Private/ShadowSetupMobile.cpp b/Engine/Source/Runtime/Renderer/Private/ShadowSetupMobile.cpp
--- a/Engine/Source/Runtime/Renderer/Private/ShadowSetupMobile.cpp
+++ b/Engine/Source/Runtime/Renderer/Private/ShadowSetupMobile.cpp
@@ -145,7 +145,6 @@ static bool MobileDetermineStaticMeshesCSMVisibilityStateInner(
 	const FLightSceneInfo& LightSceneInfo = ProjectedShadowInfo->GetLightSceneInfo();
 	FLightSceneProxy* RESTRICT LightProxy = LightSceneInfo.Proxy;
 	FVector LightDir = LightProxy->GetDirection();
-	const float ShadowCastLength = WORLD_MAX;
 
 	FPrimitiveSceneInfo* RESTRICT	PrimitiveSceneInfo = PrimitiveSceneInfoCompact.PrimitiveSceneInfo;
 	FPrimitiveSceneProxy* RESTRICT	PrimitiveProxy = PrimitiveSceneInfoCompact.Proxy;
@@ -157,10 +156,10 @@ static bool MobileDetermineStaticMeshesCSMVisibilityStateInner(
 		const FVector LightDirection = LightProxy->GetDirection();
 		const FVector PrimitiveToShadowCenter = ProjectedShadowInfo->ShadowBounds.Center - PrimitiveBounds.Origin;
 		// Project the primitive's bounds origin onto the light vector
-		const float ProjectedDistanceFromShadowOriginAlongLightDir = PrimitiveToShadowCenter | LightDirection;
+		const FVector::FReal ProjectedDistanceFromShadowOriginAlongLightDir = PrimitiveToShadowCenter | LightDirection;
 		// Calculate the primitive's squared distance to the cylinder's axis
-		const float PrimitiveDistanceFromCylinderAxisSq = (-LightDirection * ProjectedDistanceFromShadowOriginAlongLightDir + PrimitiveToShadowCenter).SizeSquared();
-		const float CombinedRadiusSq = FMath::Square(ProjectedShadowInfo->ShadowBounds.W + PrimitiveBounds.SphereRadius);
+		const FVector::FReal PrimitiveDistanceFromCylinderAxisSq = (-LightDirection * ProjectedDistanceFromShadowOriginAlongLightDir + PrimitiveToShadowCenter).SizeSquared();
+		const FVector::FReal CombinedRadiusSq = FMath::Square(ProjectedShadowInfo->ShadowBounds.W + PrimitiveBounds.SphereRadius);
 
 		// Include all primitives for movable lights, but only statically shadowed primitives from a light with static shadowing,
 		// Since lights with static shadowing still create per-object shadows for primitives without static shadowing.
@@ -231,7 +230,7 @@ static void VisualizeMobileDynamicCSMSubjectCapsules(FViewInfo& View, FLightScen
 	FVisibleLightViewInfo& VisibleLightViewInfo = View.VisibleLightInfos[LightSceneInfo->Id];
 	FMobileCSMSubjectPrimitives& MobileCSMSubjectPrimitives = VisibleLightViewInfo.MobileCSMSubjectPrimitives;
 	FVector LightDir = LightSceneInfo->Proxy->GetDirection();
-	const float ShadowCastLength = WORLD_MAX;
+	const FVector::FReal ShadowCastLength = WORLD_MAX;
 	const uint32 CullingMethod = CVarsCsmShaderCullingMethod.GetValueOnRenderThread() & 0xF;
 	const bool bSphereTest = (CVarsCsmShaderCullingMethod.GetValueOnRenderThread() & 0x10) != 0;
 
