GeometryMode: Resolve issues with non-inclusive terminology

diff --git a/Engine/Plugins/Editor/GeometryMode/Source/GeometryMode/Public/EditorGeometry.h b/Engine/Plugins/Editor/GeometryMode/Source/GeometryMode/Public/EditorGeometry.h
--- a/Engine/Plugins/Editor/GeometryMode/Source/GeometryMode/Public/EditorGeometry.h
+++ b/Engine/Plugins/Editor/GeometryMode/Source/GeometryMode/Public/EditorGeometry.h
@@ -267,7 +267,7 @@ public:
 
 	void CompileSelectionOrder();
 
-	/** Master lists.  All lower data types refer to the contents of these pools through indices. */
+	/** Main lists.  All lower data types refer to the contents of these pools through indices. */
 	TArray<FGeomVertex> VertexPool;
 	TArray<FGeomEdge> EdgePool;
 	TArray<FGeomPoly> PolyPool;
