GeometryMode - fix truncation errors.

diff --git a/Engine/Plugins/Editor/GeometryMode/Source/GeometryMode/Private/EditorGeometry.cpp b/Engine/Plugins/Editor/GeometryMode/Source/GeometryMode/Private/EditorGeometry.cpp
--- a/Engine/Plugins/Editor/GeometryMode/Source/GeometryMode/Private/EditorGeometry.cpp
+++ b/Engine/Plugins/Editor/GeometryMode/Source/GeometryMode/Private/EditorGeometry.cpp
@@ -197,7 +197,7 @@ bool FGeomEdge::IsSameEdge( const FGeomEdge& InEdge ) const
 FVector FGeomEdge::GetWidgetLocation()
 {
 	FVector dir = ((FVector)GetParentObject()->VertexPool[ VertexIndices[1] ] - (FVector)GetParentObject()->VertexPool[ VertexIndices[0] ]);
-	const float dist = dir.Size() / 2;
+	const double dist = dir.Size() / 2;
 	dir.Normalize();
 	const FVector loc = (FVector)GetParentObject()->VertexPool[ VertexIndices[0] ] + (dir * dist);
 	return GetParentObject()->GetActualBrush()->ActorToWorld().TransformPosition( loc );
diff --git a/Engine/Plugins/Editor/GeometryMode/Source/GeometryMode/Private/GeometryEdMode.cpp b/Engine/Plugins/Editor/GeometryMode/Source/GeometryMode/Private/GeometryEdMode.cpp
--- a/Engine/Plugins/Editor/GeometryMode/Source/GeometryMode/Private/GeometryEdMode.cpp
+++ b/Engine/Plugins/Editor/GeometryMode/Source/GeometryMode/Private/GeometryEdMode.cpp
@@ -666,7 +666,7 @@ void FEdModeGeometry::RenderVertex( const FSceneView* View, FPrimitiveDrawInterf
 			check(GeomObject->GetActualBrush());
 
 			Location = (FVector)GeomObject->GetActualBrush()->ActorToWorld().TransformPosition( (FVector)*GeomVertex );
-			Scale = View->WorldToScreen( Location ).W * ( 4.0f / View->UnscaledViewRect.Width() / View->ViewMatrices.GetProjectionMatrix().M[0][0] );
+			Scale = static_cast<float>( View->WorldToScreen( Location ).W * ( 4.0 / View->UnscaledViewRect.Width() / View->ViewMatrices.GetProjectionMatrix().M[0][0] ) );
 			Color = GeomVertex->IsSelected() ? FColor(255,128,64) : GeomObject->GetActualBrush()->GetWireColor();
 
 			PDI->SetHitProxy( new HGeomVertexProxy( GeomObject, VertIdx) );
diff --git a/Engine/Plugins/Editor/GeometryMode/Source/GeometryMode/Private/GeometryModifiers.cpp b/Engine/Plugins/Editor/GeometryMode/Source/GeometryMode/Private/GeometryModifiers.cpp
--- a/Engine/Plugins/Editor/GeometryMode/Source/GeometryMode/Private/GeometryModifiers.cpp
+++ b/Engine/Plugins/Editor/GeometryMode/Source/GeometryMode/Private/GeometryModifiers.cpp
@@ -60,7 +60,10 @@ static FVector ComputeWorldSpaceMousePos( FEditorViewportClient* ViewportClient
 	FSceneView* View = ViewportClient->CalcSceneView(&ViewFamily);
 
 	// Note only works for ortho viewports
-	return View->PixelToWorld(ViewportClient->Viewport->GetMouseX(),ViewportClient->Viewport->GetMouseY(),0.5f);
+	const float X = static_cast<float>(ViewportClient->Viewport->GetMouseX());
+	const float Y = static_cast<float>(ViewportClient->Viewport->GetMouseY());
+	constexpr float Z = 0.5f; 
+	return View->PixelToWorld(X, Y, Z);
 }
 
 UGeomModifier::UGeomModifier(const FObjectInitializer& ObjectInitializer)
@@ -749,8 +752,8 @@ void UGeomModifier_Extrude::Initialize()
 {
 	//the user may have changed the mode AFTER going into extrude - double check its LOCAL not WORLD
 	CheckCoordinatesMode();
-
-	Apply( GEditor->GetGridSize(), 1 );
+	const int32 GridSize = static_cast<int32>(GEditor->GetGridSize());
+	Apply( GridSize , 1 );
 }
 
 bool UGeomModifier_Extrude::OnApply()
@@ -1553,16 +1556,16 @@ static bool DoLineSegmentsIntersect(const FVector2D& Segment1Start, const FVecto
 	const FVector2D Segment1Dir = Segment1End - Segment1Start;
 	const FVector2D Segment2Dir = Segment2End - Segment2Start;
 
-	const float Determinant = FVector2D::CrossProduct(Segment1Dir, Segment2Dir);
+	const double Determinant = FVector2D::CrossProduct(Segment1Dir, Segment2Dir);
 	if (!FMath::IsNearlyZero(Determinant))
 	{
 		const FVector2D SegmentStartDelta = Segment2Start - Segment1Start;
-		const float OneOverDet = 1.0f / Determinant;
-		const float Seg1Intersection = FVector2D::CrossProduct(SegmentStartDelta, Segment2Dir) * OneOverDet;
-		const float Seg2Intersection = FVector2D::CrossProduct(SegmentStartDelta, Segment1Dir) * OneOverDet;
+		const double OneOverDet = 1.0 / Determinant;
+		const double Seg1Intersection = FVector2D::CrossProduct(SegmentStartDelta, Segment2Dir) * OneOverDet;
+		const double Seg2Intersection = FVector2D::CrossProduct(SegmentStartDelta, Segment1Dir) * OneOverDet;
 
-		const float Epsilon = 1/128.0f;
-		return (Seg1Intersection > Epsilon && Seg1Intersection < 1.0f - Epsilon && Seg2Intersection > Epsilon && Seg2Intersection < 1.0f - Epsilon);
+		const double Epsilon = 1/128.0f;
+		return (Seg1Intersection > Epsilon && Seg1Intersection < 1.0 - Epsilon && Seg2Intersection > Epsilon && Seg2Intersection < 1.0 - Epsilon);
 	}
 
 	return false;
@@ -1591,7 +1594,7 @@ static bool DoesFinalLineIntersectWithShape(const TArray<FVector>& Vertices, con
 	const FVector& Segment1End = EndVertex;
 
 	const FVector Segment1Dir = Segment1End - Segment1Start;
-	const float Segment1Len = Segment1Dir.Size();
+	const double Segment1Len = Segment1Dir.Size();
 	if (FMath::IsNearlyZero(Segment1Len))
 	{
 		// Treat zero length line segments as non-intersecting
@@ -1633,8 +1636,8 @@ static bool DoesFinalLineIntersectWithShape(const TArray<FVector>& Vertices, con
 		const FVector ProjectedYAxis = FVector::CrossProduct(PlaneNormal.GetSafeNormal(), ProjectedXAxis);
 
 		// Project 3d points onto plane
-		const FVector2D ProjectedSegment1Start(0.0f, 0.0f);
-		const FVector2D ProjectedSegment1End(Segment1Len, 0.0f);
+		const FVector2D ProjectedSegment1Start(0.0, 0.0);
+		const FVector2D ProjectedSegment1End(Segment1Len, 0.0);
 		const FVector2D ProjectedSegment2Start(FVector::DotProduct(ProjectedXAxis, SegmentStartDelta), FVector::DotProduct(ProjectedYAxis, SegmentStartDelta));
 		const FVector2D ProjectedSegment2End(FVector::DotProduct(ProjectedXAxis, Segment2End - Segment1Start), FVector::DotProduct(ProjectedYAxis, Segment2End - Segment1Start));
 
@@ -1773,7 +1776,7 @@ void UGeomModifier_Pen::Render(const FSceneView* View,FViewport* Viewport,FPrimi
 
 	// Draw a box where the next vertex will be placed
 
-	int32 BoxSz = FMath::Max( GEditor->GetGridSize() / 2, 1.f );
+	int32 BoxSz = static_cast<int32>( FMath::Max( GEditor->GetGridSize() / 2, 1.f ) );
 	DrawWireBox(PDI, FBox::BuildAABB( MouseWorldSpacePos, FVector(BoxSz,BoxSz,BoxSz) ), FLinearColor(1,1,1), SDPG_Foreground);
 }
 
@@ -2431,9 +2434,9 @@ void UGeomModifier_Clip::Render(const FSceneView* View,FViewport* Viewport,FPrim
 
 		if( !NormalPoly.CalcNormal(1) )
 		{
-			FVector Start = ( vtx1 + vtx2 ) / 2.f;
+			FVector Start = ( vtx1 + vtx2 ) / 2.;
 
-			float NormalLength = (vtx2 - vtx1).Size() / 2.f;
+			double NormalLength = (vtx2 - vtx1).Size() / 2.;
 
 			if( ClipMarkers.Num() == 1 )
 			{
@@ -2448,7 +2451,7 @@ void UGeomModifier_Clip::Render(const FSceneView* View,FViewport* Viewport,FPrim
 
 	// Draw a box at the cursor location
 
-	int32 BoxSz = FMath::Max( GEditor->GetGridSize() / 2, 1.f );
+	int32 BoxSz = static_cast<int32>( FMath::Max( GEditor->GetGridSize() / 2, 1.f ) );
 	DrawWireBox(PDI, FBox::BuildAABB( SnappedMouseWorldSpacePos, FVector(BoxSz,BoxSz,BoxSz) ), FLinearColor(1,1,1), SDPG_Foreground);
 }
 
