TextureAlign Mode - fix truncation warnings

diff --git a/Engine/Plugins/Editor/GeometryMode/Source/TextureAlignMode/Private/TextureAlignEdMode.cpp b/Engine/Plugins/Editor/GeometryMode/Source/TextureAlignMode/Private/TextureAlignEdMode.cpp
--- a/Engine/Plugins/Editor/GeometryMode/Source/TextureAlignMode/Private/TextureAlignEdMode.cpp
+++ b/Engine/Plugins/Editor/GeometryMode/Source/TextureAlignMode/Private/TextureAlignEdMode.cpp
@@ -245,7 +245,10 @@ bool FModeTool_Texture::InputDelta(FEditorViewportClient* InViewportClient,FView
 			}
 			FMatrix Mat = GLevelEditorModeTools().GetCustomDrawingCoordinateSystem();
 			FVector UVW = Mat.InverseTransformVector( deltaDrag );  // InverseTransformNormal because Mat is the transform from the surface/widget's coords to world coords
-			GEditor->polyTexPan( Model, UVW.X, UVW.Y, 0 );  // 0 is relative mode because UVW is made from deltaDrag - the user input since the last tick
+			const int32 X = static_cast<int32>(UVW.X);
+			const int32 Y = static_cast<int32>(UVW.Y);
+			constexpr int32 Z = 0;
+			GEditor->polyTexPan( Model, X, Y, Z );  // 0 is relative mode because UVW is made from deltaDrag - the user input since the last tick
 		}
 	}
 
@@ -276,8 +279,8 @@ bool FModeTool_Texture::InputDelta(FEditorViewportClient* InViewportClient,FView
 
 	if( !InScale.IsZero() )
 	{
-		float ScaleU = InScale.X / GEditor->GetGridSize();
-		float ScaleV = InScale.Y / GEditor->GetGridSize();
+		float ScaleU = static_cast<float>( InScale.X / GEditor->GetGridSize() );
+		float ScaleV = static_cast<float>( InScale.Y / GEditor->GetGridSize() );
 
 		ScaleU = 1.f - (ScaleU / 100.f);
 		ScaleV = 1.f - (ScaleV / 100.f);
