Geometry mode in place actors looks different and doesn't support favorites or filtering (Support UObject types for Place Actors)

- altered BSP drag drop op to get some consistency with the asset drag drop
- gave the BSP brushes a drag handler
- altered the FPlaceableItems to cater to custom draggables
- altered the placement mode favorites and filtering handling to incorporate the new supported draggables

diff --git a/Engine/Plugins/Editor/GeometryMode/Source/BspMode/Private/BspDragHandler.cpp b/Engine/Plugins/Editor/GeometryMode/Source/BspMode/Private/BspDragHandler.cpp
new file mode 100644
--- /dev/null
+++ b/Engine/Plugins/Editor/GeometryMode/Source/BspMode/Private/BspDragHandler.cpp
@@ -0,0 +1,24 @@
+// Copyright Epic Games, Inc. All Rights Reserved.
+
+#include "BspDragHandler.h"
+#include "EditorClassUtils.h"
+#include "BspModeModule.h"
+#include "DragAndDrop/BrushBuilderDragDropOp.h"
+
+FBspDragHandler::FBspDragHandler()
+{
+}
+
+void FBspDragHandler::Initialize( TSharedRef<FBspBuilderType> BspBuilder )
+{
+	ToolTip = FEditorClassUtils::GetTooltip(ABrush::StaticClass(), BspBuilder->ToolTipText );
+	IconBrush = BspBuilder->Icon; 
+	
+	GetContentToDrag.BindLambda( [ BspBuilder ] ()
+	{
+		const bool bIsAdditive = true;
+		
+		const TWeakObjectPtr<UBrushBuilder> ActiveBrushBuilder = GEditor->FindBrushBuilder(BspBuilder->BuilderClass.Get());
+		return FBrushBuilderDragDropOp::New( ActiveBrushBuilder, BspBuilder->Icon, bIsAdditive );
+	});
+}
\ No newline at end of file
diff --git a/Engine/Plugins/Editor/GeometryMode/Source/BspMode/Private/BspDragHandler.h b/Engine/Plugins/Editor/GeometryMode/Source/BspMode/Private/BspDragHandler.h
new file mode 100644
--- /dev/null
+++ b/Engine/Plugins/Editor/GeometryMode/Source/BspMode/Private/BspDragHandler.h
@@ -0,0 +1,15 @@
+// Copyright Epic Games, Inc. All Rights Reserved.
+
+#pragma once
+
+#include "IPlacementModeModule.h"
+
+struct FBspBuilderType;
+
+class FBspDragHandler : public FPlaceableItem::FDragHandler
+{
+public:
+
+	FBspDragHandler();
+	void Initialize( TSharedRef<FBspBuilderType> InBspBuilder );
+};
diff --git a/Engine/Plugins/Editor/GeometryMode/Source/BspMode/Private/BspModeModule.cpp b/Engine/Plugins/Editor/GeometryMode/Source/BspMode/Private/BspModeModule.cpp
--- a/Engine/Plugins/Editor/GeometryMode/Source/BspMode/Private/BspModeModule.cpp
+++ b/Engine/Plugins/Editor/GeometryMode/Source/BspMode/Private/BspModeModule.cpp
@@ -12,13 +12,28 @@
 #include "BspModeStyle.h"
 #include "SBspPalette.h"
 #include "IPlacementModeModule.h"
+#include "DragAndDrop/BrushBuilderDragDropOp.h"
+#include "BspDragHandler.h"
+#include "IPlacementModeModule.h"
 
 #define LOCTEXT_NAMESPACE "BspMode"
 
+FBspBuilderType::FBspBuilderType(
+	UClass* InBuilderClass,
+	const FText& InText,
+	const FText& InToolTipText,
+	const FSlateBrush* InIcon ):
+BuilderClass(InBuilderClass)
+	, Text(InText)
+	, ToolTipText(InToolTipText)
+	, Icon(InIcon)
+{
+}
+
 void FBspModeModule::StartupModule()
 {
 	FBspModeStyle::Initialize();
-
+	CategoryName = "Geometry";
 /*
 	FEditorModeRegistry::Get().RegisterMode<FBspMode>(
 		FGeometryEditingModes::EM_Bsp,
@@ -29,6 +44,10 @@ void FBspModeModule::StartupModule()
 		);
 */
 
+	FPlacementCategoryInfo CategoryInfo (LOCTEXT("PlacementMode_Geometry", "Geometry"), FSlateIcon(FAppStyle::GetAppStyleSetName(), "PlacementBrowser.Icons.BSP"), CategoryName, TEXT("PMGeometry"), 35);
+	CategoryInfo.CustomGenerator = []() -> TSharedRef<SWidget> { return SNew(SBspPalette); };
+	IPlacementModeModule::Get().RegisterPlacementCategory( CategoryInfo );
+	
 	RegisterBspBuilderType(UCubeBuilder::StaticClass(), LOCTEXT("CubeBuilderName", "Box"), LOCTEXT("CubeBuilderToolTip", "Make a box brush"), FBspModeStyle::Get().GetBrush(TEXT("BspMode.BoxBrush")));
 	RegisterBspBuilderType(UConeBuilder::StaticClass(), LOCTEXT("ConeBuilderName", "Cone"), LOCTEXT("ConeBuilderToolTip", "Make a cone brush"), FBspModeStyle::Get().GetBrush(TEXT("BspMode.ConeBrush")));
 	RegisterBspBuilderType(UCylinderBuilder::StaticClass(), LOCTEXT("CylinderBuilderName", "Cylinder"), LOCTEXT("CylinderBuilderToolTip", "Make a cylinder brush"), FBspModeStyle::Get().GetBrush(TEXT("BspMode.CylinderBrush")));
@@ -36,19 +55,18 @@ void FBspModeModule::StartupModule()
 	RegisterBspBuilderType(ULinearStairBuilder::StaticClass(), LOCTEXT("LinearStairBuilderName", "Linear Stair"), LOCTEXT("LinearStairBuilderToolTip", "Make a linear stair brush"), FBspModeStyle::Get().GetBrush(TEXT("BspMode.LinearStairBrush")));
 	RegisterBspBuilderType(USpiralStairBuilder::StaticClass(), LOCTEXT("SpiralStairBuilderName", "Spiral Stair"), LOCTEXT("SpiralStairBuilderToolTip", "Make a spiral stair brush"), FBspModeStyle::Get().GetBrush(TEXT("BspMode.SpiralStairBrush")));
 	RegisterBspBuilderType(UTetrahedronBuilder::StaticClass(), LOCTEXT("SphereBuilderName", "Sphere"), LOCTEXT("SphereBuilderToolTip", "Make a sphere brush"), FBspModeStyle::Get().GetBrush(TEXT("BspMode.SphereBrush")));
-
-	IPlacementModeModule& PlacementModeModule = IPlacementModeModule::Get();
-
-	FPlacementCategoryInfo CategoryInfo(LOCTEXT("PlacementMode_Geometry", "Geometry"), FSlateIcon(FAppStyle::GetAppStyleSetName(), "PlacementBrowser.Icons.BSP"), "Geometry", TEXT("PMGeometry"), 35);
-	CategoryInfo.CustomGenerator = []() -> TSharedRef<SWidget> { return SNew(SBspPalette); };
-	PlacementModeModule.RegisterPlacementCategory(CategoryInfo);
 }
 
-
 void FBspModeModule::ShutdownModule()
 {
-	//FEditorModeRegistry::Get().UnregisterMode(FGeometryEditingModes::EM_Bsp);
-
+	IPlacementModeModule::Get().UnregisterPlacementCategory( CategoryName );
+	UnregisterBspBuilderType(UCubeBuilder::StaticClass());
+	UnregisterBspBuilderType(UConeBuilder::StaticClass());
+	UnregisterBspBuilderType(UCylinderBuilder::StaticClass());
+	UnregisterBspBuilderType(UCurvedStairBuilder::StaticClass());
+	UnregisterBspBuilderType(ULinearStairBuilder::StaticClass());
+	UnregisterBspBuilderType(USpiralStairBuilder::StaticClass());
+	UnregisterBspBuilderType(UTetrahedronBuilder::StaticClass()); 
 	BspBuilderTypes.Empty();
 }
 
@@ -56,12 +74,32 @@ void FBspModeModule::ShutdownModule()
 void FBspModeModule::RegisterBspBuilderType( class UClass* InBuilderClass, const FText& InBuilderName, const FText& InBuilderTooltip, const FSlateBrush* InBuilderIcon )
 {
 	check(InBuilderClass->IsChildOf(UBrushBuilder::StaticClass()));
-	BspBuilderTypes.Add(MakeShareable(new FBspBuilderType(InBuilderClass, InBuilderName, InBuilderTooltip, InBuilderIcon)));
+	const TSharedPtr<FBspBuilderType> BuilderType = MakeShareable(new FBspBuilderType(InBuilderClass, InBuilderName, InBuilderTooltip, InBuilderIcon));
+	BspBuilderTypes.Add( BuilderType );
+	const TSharedPtr<FBspDragHandler> Handler = MakeShared<FBspDragHandler>();
+	Handler->Initialize( BuilderType.ToSharedRef() );
+	const FString Name = "BSP_" + BuilderType->Text.ToString();
+
+	static int32 SortOrder = 0;
+	
+	BuilderType->PlaceableItem =  MakeShared<FPlaceableItem>( Handler, SortOrder++, BuilderType->Text, Name );
+	BuilderType->PlacementModeID = IPlacementModeModule::Get().RegisterPlaceableItem( CategoryName,   BuilderType->PlaceableItem.ToSharedRef() );
 }
 
 
 void FBspModeModule::UnregisterBspBuilderType( class UClass* InBuilderClass )
 {
+	for (  TSharedPtr<FBspBuilderType> BuilderType : BspBuilderTypes )
+	{
+		if ( BuilderType.IsValid() && BuilderType->PlaceableItem.IsValid() )
+		{
+			if ( BuilderType->PlacementModeID.IsSet() )
+			{
+				IPlacementModeModule::Get().UnregisterPlaceableItem( BuilderType->PlacementModeID.GetValue() );
+			}
+		}
+	}
+	
 	BspBuilderTypes.RemoveAll( 
 		[InBuilderClass] ( const TSharedPtr<FBspBuilderType>& RemovalCandidate ) -> bool
 		{
diff --git a/Engine/Plugins/Editor/GeometryMode/Source/BspMode/Private/BspModeModule.h b/Engine/Plugins/Editor/GeometryMode/Source/BspMode/Private/BspModeModule.h
--- a/Engine/Plugins/Editor/GeometryMode/Source/BspMode/Private/BspModeModule.h
+++ b/Engine/Plugins/Editor/GeometryMode/Source/BspMode/Private/BspModeModule.h
@@ -5,18 +5,18 @@
 #include "UObject/Class.h"
 #include "UObject/WeakObjectPtr.h"
 #include "IBspModeModule.h"
+#include "IPlacementModeModule.h"
+#include "Templates/SharedPointer.h"
+#include "Misc/Optional.h"
 
 struct FSlateBrush;
+struct FPlaceableItem;
+struct FDragHandler;
+class FBrushBuilderDragDropOp;
 
 struct FBspBuilderType
 {
-	FBspBuilderType(class UClass* InBuilderClass, const FText& InText, const FText& InToolTipText, const FSlateBrush* InIcon)
-		: BuilderClass(InBuilderClass)
-		, Text(InText)
-		, ToolTipText(InToolTipText)
-		, Icon(InIcon)
-	{
-	}
+	FBspBuilderType(class UClass* InBuilderClass, const FText& InText, const FText& InToolTipText, const FSlateBrush* InIcon);
 
 	/** The class of the builder brush */
 	TWeakObjectPtr<UClass> BuilderClass;
@@ -29,6 +29,12 @@ struct FBspBuilderType
 
 	/** The icon to be displayed for this builder */
 	const FSlateBrush* Icon;
+
+	/** the placeable item that will provide the information for the draggable  for this BSP builder */
+	TSharedPtr<FPlaceableItem> PlaceableItem;
+
+	/** TOptional<FPlacementModeID>for deletion */
+	TOptional<FPlacementModeID> PlacementModeID;
 };
 
 class FBspModeModule : public IBspModeModule
@@ -49,4 +55,5 @@ public:
 private:
 
 	TArray< TSharedPtr<FBspBuilderType> > BspBuilderTypes;
+	FName CategoryName;
 };
diff --git a/Engine/Source/Editor/PlacementMode/Private/PlacementModeModule.cpp b/Engine/Source/Editor/PlacementMode/Private/PlacementModeModule.cpp
--- a/Engine/Source/Editor/PlacementMode/Private/PlacementModeModule.cpp
+++ b/Engine/Source/Editor/PlacementMode/Private/PlacementModeModule.cpp
@@ -570,14 +570,21 @@ TOptional<FPlacementModeID> FPlacementModeModule::RegisterPlaceableItem(FName Ca
 	using namespace PlacementModeModuleLocals;
 
 	FPlacementCategory* Category = Categories.Find(CategoryName);
-	if (Category && !Category->CustomGenerator)
+	if (Category )
 	{
-		FPlacementModeID ID = CreateID(CategoryName);
-		Category->Items.Add(ID.UniqueID, InItem);
+		if ( InItem->DragHandler.IsValid() )
+		{
+			Category->CustomDraggableItems.Add( InItem );
+		}
+		if (!Category->CustomGenerator)
+		{
+			FPlacementModeID ID = CreateID(CategoryName);
+			Category->Items.Add(ID.UniqueID, InItem);
 
-		ManuallyCreatedPlaceableItems.Add(MakePlacementInfo(*InItem), InItem);
+			ManuallyCreatedPlaceableItems.Add(MakePlacementInfo(*InItem), InItem);
 
-		return ID;
+			return ID;
+		}
 	}
 	return TOptional<FPlacementModeID>();
 }
@@ -593,6 +600,7 @@ void FPlacementModeModule::UnregisterPlaceableItem(FPlacementModeID ID)
 		Category->Items.RemoveAndCopyValue(ID.UniqueID, Item);
 		if (Item)
 		{
+			Category->CustomDraggableItems.Remove(Item.ToSharedRef());
 			ManuallyCreatedPlaceableItems.Remove(MakePlacementInfo(*Item));
 		}
 	}
@@ -797,6 +805,15 @@ void FPlacementModeModule::RefreshAllPlaceableClasses()
 
 	Category->Items.Reset();
 
+
+	for ( TTuple<FName, FPlacementCategory>& Pair : Categories )
+	{
+		for ( const TSharedRef<FPlaceableItem>& Draggable : Pair.Value.CustomDraggableItems  )
+		{
+			Category->Items.Add( CreateID(), Draggable.ToSharedPtr() );
+		}
+	}
+
 	// Manually add some special cases that aren't added below
 	Category->Items.Add(CreateID(), MakeShareable(new FPlaceableItem(*UActorFactoryEmptyActor::StaticClass())));
 	Category->Items.Add(CreateID(), MakeShareable(new FPlaceableItem(*UActorFactoryCharacter::StaticClass())));
diff --git a/Engine/Source/Editor/PlacementMode/Private/SPlacementModeTools.cpp b/Engine/Source/Editor/PlacementMode/Private/SPlacementModeTools.cpp
--- a/Engine/Source/Editor/PlacementMode/Private/SPlacementModeTools.cpp
+++ b/Engine/Source/Editor/PlacementMode/Private/SPlacementModeTools.cpp
@@ -105,6 +105,7 @@ public:
 		, _Height( 32 )
 		, _AlwaysUseGenericThumbnail( false )
 		, _AssetTypeColorOverride()
+		, _CustomIconBrush( nullptr )
 	{}
 
 	SLATE_ARGUMENT( uint32, Width )
@@ -116,8 +117,10 @@ public:
 	SLATE_ARGUMENT( bool, AlwaysUseGenericThumbnail )
 
 	SLATE_ARGUMENT( TOptional<FLinearColor>, AssetTypeColorOverride )
-	SLATE_END_ARGS()
+	
+	SLATE_ARGUMENT( const FSlateBrush*, CustomIconBrush )
 
+	SLATE_END_ARGS()
 
 	void Construct( const FArguments& InArgs, const FAssetData& InAsset)
 	{
@@ -139,6 +142,11 @@ public:
 				FClassIconFinder::FindThumbnailForClass( nullptr,  InArgs._ClassThumbnailBrushOverride ) :
 				FClassIconFinder::FindThumbnailForClass( ThumbnailClass.Get(), DefaultThumbnail );
 
+			if ( InArgs._CustomIconBrush )
+			{
+				ThumbnailBrush = InArgs._CustomIconBrush;
+			}
+
 			ChildSlot[ SAssignNew( ThumbnailImage, SImage ).Image( ThumbnailBrush ) ];
 		}
 		else
@@ -247,6 +255,12 @@ void SPlacementAssetEntry::Construct(const FArguments& InArgs, const TSharedPtr<
 		ThumbnailBoxWidth = 20;
 	}
 	
+	const FSlateBrush* CustomIconBrush = nullptr;
+	if ( !bIsClassicView && Item->DragHandler.IsValid() && Item->DragHandler->IconBrush )
+	{
+		CustomIconBrush = Item->DragHandler->IconBrush;
+	}
+	
 	ChildSlot
 	.Padding( WholeAssetPadding )
 	[
@@ -278,6 +292,7 @@ void SPlacementAssetEntry::Construct(const FArguments& InArgs, const TSharedPtr<
 						.ClassThumbnailBrushOverride( Item->ClassThumbnailBrushOverride )
 						.AlwaysUseGenericThumbnail( Item->bAlwaysUseGenericThumbnail )
 						.AssetTypeColorOverride( FLinearColor::Transparent )
+						.CustomIconBrush( CustomIconBrush )
 					]
 				]
 
@@ -372,6 +387,11 @@ FReply SPlacementAssetEntry::OnDragDetected(const FGeometry& MyGeometry, const F
 
 	if( MouseEvent.IsMouseButtonDown( EKeys::LeftMouseButton ) )
 	{
+		if ( Item->DragHandler.IsValid()  && Item->DragHandler->GetContentToDrag.IsBound())
+		{
+			return FReply::Handled().BeginDragDrop( Item->DragHandler->GetContentToDrag.Execute() );
+		}
+		
 		return FReply::Handled().BeginDragDrop(FAssetDragDropOp::New(Item->AssetData, Item->AssetFactory));
 	}
 	else
@@ -832,6 +852,7 @@ void SPlacementModeTools::SetActiveTab(FName TabName)
 void SPlacementModeTools::UpdateShownItems()
 {
 	bUpdateShownItems = false;
+	bool bIsCategoryContentChooserView = CVarEnableCategoryContentChooserView->GetBool();
 
 	IPlacementModeModule& PlacementModeModule = IPlacementModeModule::Get();
 
@@ -840,7 +861,7 @@ void SPlacementModeTools::UpdateShownItems()
 	{
 		return;
 	}
-	else if (Category->CustomGenerator)
+	else if (Category->CustomGenerator && (Category->CustomDraggableItems.IsEmpty() || !bIsCategoryContentChooserView))
 	{
 		CategoryContentBuilder->FillWithBuilder( Category->CustomGenerator() );
 		CustomContent->SetContent(Category->CustomGenerator());
@@ -871,7 +892,17 @@ void SPlacementModeTools::UpdateShownItems()
 		}
 		else
 		{
-			PlacementModeModule.GetItemsForCategory(Category->UniqueHandle, FilteredItems);
+			if ( !Category->CustomDraggableItems.IsEmpty() )
+			{
+				for (TSharedRef<FPlaceableItem> Item : Category->CustomDraggableItems)
+				{
+					FilteredItems.Add( Item.ToSharedPtr() );
+				}
+			}
+			else
+			{
+				PlacementModeModule.GetItemsForCategory(Category->UniqueHandle, FilteredItems);
+			}
 
 			if (Category->bSortable)
 			{
@@ -954,7 +985,7 @@ void SPlacementModeTools::UpdateContentForCategory( FName CategoryName, FText Ca
 		UpdateShownItems();
 
 		const FPlacementCategoryInfo* Category = IPlacementModeModule::Get().GetRegisteredPlacementCategory( CategoryName );
-		if ( Category && Category->CustomGenerator )
+		if ( Category && Category->CustomGenerator && Category->CustomDraggableItems.IsEmpty() )
 		{
 			CategoryContentBuilder->FillWithBuilder( Category->CustomGenerator() );
 			return;
diff --git a/Engine/Source/Editor/PlacementMode/Public/IPlacementModeModule.h b/Engine/Source/Editor/PlacementMode/Public/IPlacementModeModule.h
--- a/Engine/Source/Editor/PlacementMode/Public/IPlacementModeModule.h
+++ b/Engine/Source/Editor/PlacementMode/Public/IPlacementModeModule.h
@@ -17,6 +17,7 @@
 
 class FNamePermissionList;
 class FBuilderInput;
+struct FPlaceableItem;
 
 /**
  * Struct that defines an identifier for a particular placeable item in this module.
@@ -76,6 +77,11 @@ struct FPlacementCategoryInfo
 
 	/** Whether the items in this category are automatically sortable by name. False if the items are already sorted. */
 	bool bSortable;
+
+	/**
+	 * FPlaceableItems with custom drag handling, populated only if the category handles making
+	 * its own draggable items */
+	TArray< TSharedRef<FPlaceableItem>>  CustomDraggableItems;
 };
 
 /**
@@ -83,6 +89,24 @@ struct FPlacementCategoryInfo
  */
 struct FPlaceableItem
 {
+	/**
+	* An object which provides handling for Drags.
+	*/
+	struct FDragHandler : public TSharedFromThis<FDragHandler>
+	{
+	public:
+		DECLARE_DELEGATE_RetVal(TSharedRef<FDragDropOperation>, FGetContentToDrag  );
+
+		/** a delegate thart returns the FDragDropOperation for the  the draggable */
+		FGetContentToDrag GetContentToDrag;
+
+		/** the tooltip for the draggable */
+		TSharedPtr<IToolTip> ToolTip;
+		
+		/** the const FSlateBrush* that provides the icon for the draggable */
+		const FSlateBrush* IconBrush;
+	};
+	
 	/** Default constructor */
 	FPlaceableItem()
 		: Factory(nullptr)
@@ -154,6 +178,19 @@ struct FPlaceableItem
 		}
 	}
 
+	FPlaceableItem(
+		TSharedPtr<FPlaceableItem::FDragHandler> InDragHandler
+		, TOptional<int32> InSortOrder
+		, FText InLabel
+		, FString InName ):
+		NativeName( InName )
+		, DisplayName( InLabel )
+		, bAlwaysUseGenericThumbnail(false)
+		, SortOrder(InSortOrder)
+		, DragHandler(InDragHandler)
+	{
+	}
+
 	/** Constructor for any placeable actor class with associated asset data, brush and display name overrides */
 	 FPlaceableItem(
 		UClass& InActorFactoryClass,
@@ -242,6 +279,9 @@ public:
 	/** Optional sort order (lowest first). Overrides default class name sorting. */
 	TOptional<int32> SortOrder;
 
+	/** If provided, handles the drag for the item. This should be provided for objects which are not actors or assets, which have built in handling */
+	TSharedPtr<FDragHandler> DragHandler;
+
 private:
 
 	/** This item's native name as an FName (initialized on access only) */
