Fixed an issue with brush vertices where they had their actor transform applied twice when computing the gizmo orientation. This made the local space gizmo incorrectly aligned.
cl 2370697 added the actor transform to FEdModeGeometry::GetCustomDrawingCoordinateSystem, so it became redundant to have it in FGeomVertex::GetWidgetRotation.

diff --git a/Engine/Plugins/Editor/GeometryMode/Source/GeometryMode/Private/EditorGeometry.cpp b/Engine/Plugins/Editor/GeometryMode/Source/GeometryMode/Private/EditorGeometry.cpp
--- a/Engine/Plugins/Editor/GeometryMode/Source/GeometryMode/Private/EditorGeometry.cpp
+++ b/Engine/Plugins/Editor/GeometryMode/Source/GeometryMode/Private/EditorGeometry.cpp
@@ -171,7 +171,7 @@ FVector3f* FGeomVertex::GetActualVertex( FPolyVertexIndex& InPVI )
 
 FRotator FGeomVertex::GetWidgetRotation() const
 {
-	return GetParentObject()->GetActualBrush()->GetActorRotation();
+	return FRotator::ZeroRotator;
 }
 
 ///////////////////////////////////////////////////////////////////////////////
