GeometryMode: Fixed erroneous usage of StaticClass in ShutdownModule.

#Synced-CL 48172057

diff --git a/Engine/Plugins/Editor/GeometryMode/Source/BspMode/Private/BspModeModule.cpp b/Engine/Plugins/Editor/GeometryMode/Source/BspMode/Private/BspModeModule.cpp
--- a/Engine/Plugins/Editor/GeometryMode/Source/BspMode/Private/BspModeModule.cpp
+++ b/Engine/Plugins/Editor/GeometryMode/Source/BspMode/Private/BspModeModule.cpp
@@ -21,11 +21,12 @@ FBspBuilderType::FBspBuilderType(
 	UClass* InBuilderClass,
 	const FText& InText,
 	const FText& InToolTipText,
-	const FSlateBrush* InIcon ):
-BuilderClass(InBuilderClass)
-	, Text(InText)
-	, ToolTipText(InToolTipText)
-	, Icon(InIcon)
+	const FSlateBrush* InIcon )
+:	BuilderClass(InBuilderClass)
+,	BuilderClassName(InBuilderClass ? InBuilderClass->GetFName() : NAME_None)
+,	Text(InText)
+,	ToolTipText(InToolTipText)
+,	Icon(InIcon)
 {
 }
 
@@ -58,14 +59,12 @@ void FBspModeModule::StartupModule()
 void FBspModeModule::ShutdownModule()
 {
 	IPlacementModeModule::Get().UnregisterPlacementCategory( CategoryName );
-	UnregisterBspBuilderType(UCubeBuilder::StaticClass());
-	UnregisterBspBuilderType(UConeBuilder::StaticClass());
-	UnregisterBspBuilderType(UCylinderBuilder::StaticClass());
-	UnregisterBspBuilderType(UCurvedStairBuilder::StaticClass());
-	UnregisterBspBuilderType(ULinearStairBuilder::StaticClass());
-	UnregisterBspBuilderType(USpiralStairBuilder::StaticClass());
-	UnregisterBspBuilderType(UTetrahedronBuilder::StaticClass()); 
+	for (const FName BuilderClassName : BspBuilderTypesToUnregisterOnShutdown)
+	{
+		UnregisterBspBuilderType(BuilderClassName);
+	}
 	BspBuilderTypes.Empty();
+	BspBuilderTypesToUnregisterOnShutdown.Empty();
 }
 
 
@@ -82,14 +81,16 @@ void FBspModeModule::RegisterBspBuilderType( class UClass* InBuilderClass, const
 	
 	BuilderType->PlaceableItem =  MakeShared<FPlaceableItem>( Handler, SortOrder++, BuilderType->Text, Name );
 	BuilderType->PlacementModeID = IPlacementModeModule::Get().RegisterPlaceableItem( CategoryName, BuilderType->PlaceableItem.ToSharedRef() );
+
+	BspBuilderTypesToUnregisterOnShutdown.Add( BuilderType->BuilderClassName );
 }
 
 
-void FBspModeModule::UnregisterBspBuilderType( class UClass* InBuilderClass )
+void FBspModeModule::UnregisterBspBuilderType( const FName InBuilderClassName )
 {
-	for (  TSharedPtr<FBspBuilderType> BuilderType : BspBuilderTypes )
+	for ( TSharedPtr<FBspBuilderType> BuilderType : BspBuilderTypes )
 	{
-		if ( BuilderType.IsValid() && BuilderType->PlaceableItem.IsValid() )
+		if ( BuilderType.IsValid() && BuilderType->PlaceableItem.IsValid() && BuilderType->BuilderClassName == InBuilderClassName )
 		{
 			if ( BuilderType->PlacementModeID.IsSet() )
 			{
@@ -97,14 +98,14 @@ void FBspModeModule::UnregisterBspBuilderType( class UClass* InBuilderClass )
 			}
 		}
 	}
-	// If UObject system is not intialized the blow code can not run because it creates weak object pointers which crashes
+	// If UObject system is not initialized the below code cannot run because it creates weak object pointers which crashes
 	// the internal UObject array.
 	if (UObjectInitialized())
 	{
 		BspBuilderTypes.RemoveAll(
-			[InBuilderClass](const TSharedPtr<FBspBuilderType>& RemovalCandidate) -> bool
+			[InBuilderClassName](const TSharedPtr<FBspBuilderType>& RemovalCandidate) -> bool
 			{
-				return (RemovalCandidate->BuilderClass == InBuilderClass);
+				return RemovalCandidate.IsValid() && RemovalCandidate->BuilderClassName == InBuilderClassName;
 			}
 		);
 	}
diff --git a/Engine/Plugins/Editor/GeometryMode/Source/BspMode/Private/BspModeModule.h b/Engine/Plugins/Editor/GeometryMode/Source/BspMode/Private/BspModeModule.h
--- a/Engine/Plugins/Editor/GeometryMode/Source/BspMode/Private/BspModeModule.h
+++ b/Engine/Plugins/Editor/GeometryMode/Source/BspMode/Private/BspModeModule.h
@@ -21,6 +21,9 @@ struct FBspBuilderType
 	/** The class of the builder brush */
 	TWeakObjectPtr<UClass> BuilderClass;
 
+	/** The class name of the builder brush */
+	FName BuilderClassName;
+
 	/** The name to be displayed for this builder */
 	FText Text;
 
@@ -47,7 +50,7 @@ public:
 
 	/** IBspModeModule interface */
 	virtual void RegisterBspBuilderType( class UClass* InBuilderClass, const FText& InBuilderName, const FText& InBuilderTooltip, const FSlateBrush* InBuilderIcon ) override;
-	virtual void UnregisterBspBuilderType( class UClass* InBuilderClass ) override;
+	virtual void UnregisterBspBuilderType( class FName InBuilderClassName ) override;
 
 	const TArray< TSharedPtr<FBspBuilderType> >& GetBspBuilderTypes();
 	TSharedPtr<FBspBuilderType> FindBspBuilderType(UClass* InBuilderClass) const;
@@ -56,4 +59,6 @@ private:
 
 	TArray< TSharedPtr<FBspBuilderType> > BspBuilderTypes;
 	FName CategoryName;
+	
+	TArray<FName> BspBuilderTypesToUnregisterOnShutdown;
 };
diff --git a/Engine/Plugins/Editor/GeometryMode/Source/BspMode/Public/IBspModeModule.h b/Engine/Plugins/Editor/GeometryMode/Source/BspMode/Public/IBspModeModule.h
--- a/Engine/Plugins/Editor/GeometryMode/Source/BspMode/Public/IBspModeModule.h
+++ b/Engine/Plugins/Editor/GeometryMode/Source/BspMode/Public/IBspModeModule.h
@@ -2,6 +2,7 @@
 #pragma once
 
 #include "Modules/ModuleInterface.h"
+#include "UObject/NameTypes.h"
 
 class FText;
 
@@ -11,5 +12,7 @@ class IBspModeModule : public IModuleInterface
 {
 public:
 	virtual void RegisterBspBuilderType( class UClass* InBuilderClass, const FText& InBuilderName, const FText& InBuilderTooltip, const FSlateBrush* InBuilderIcon ) = 0;
-	virtual void UnregisterBspBuilderType( class UClass* InBuilderClass ) = 0;
+	UE_DEPRECATED(5.8, "Use the FName variant of UnregisterBspBuilderType")
+	virtual void UnregisterBspBuilderType( class UClass* InBuilderClass ) { }
+	virtual void UnregisterBspBuilderType( FName InBuilderClassName ) = 0;
 };
